Move Simplygon SDK binary to NotForLicensees

diff --git a/Engine/Source/Developer/SimplygonMeshReduction/Private/SimplygonMeshReduction.cpp b/Engine/Source/Developer/SimplygonMeshReduction/Private/SimplygonMeshReduction.cpp
--- a/Engine/Source/Developer/SimplygonMeshReduction/Private/SimplygonMeshReduction.cpp
+++ b/Engine/Source/Developer/SimplygonMeshReduction/Private/SimplygonMeshReduction.cpp
@@ -345,7 +345,7 @@ public:
  
 		GetInterfaceVersionSimplygonSDKPtr GetInterfaceVersionSimplygonSDK;
 		InitializeSimplygonSDKPtr InitializeSimplygonSDK;
-		FString DllPath(FPaths::Combine(*FPaths::EngineDir(), TEXT("Binaries/ThirdParty/NoRedist/Simplygon")));
+		FString DllPath(FPaths::Combine(*FPaths::EngineDir(), TEXT("Binaries/ThirdParty/NotForLicensees/Simplygon")));
 		FString DllFilename(FPaths::Combine(*DllPath, TEXT("SimplygonSDKEpicUE4Releasex64.dll")));
  
 		void* DLLHandle = FPlatformProcess::GetDllHandle(*DllFilename);
