Remove an IsRocket() check for Simplygon. Previous log output containing 'Error:' when Simplygon DLL was not found would trip build machines into detecting a build failure. Now fail gracefully if the DLL is not present, but error if the DLL is present but cannot be loaded.

diff --git a/Engine/Source/Developer/SimplygonMeshReduction/Private/SimplygonMeshReduction.cpp b/Engine/Source/Developer/SimplygonMeshReduction/Private/SimplygonMeshReduction.cpp
--- a/Engine/Source/Developer/SimplygonMeshReduction/Private/SimplygonMeshReduction.cpp
+++ b/Engine/Source/Developer/SimplygonMeshReduction/Private/SimplygonMeshReduction.cpp
@@ -321,55 +321,47 @@ public:
 	}
 	static FSimplygonMeshReduction* Create()
 	{
-		SimplygonSDK::ISimplygonSDK* SDK = NULL;
-		ANSICHAR VersionHash[200];
-
-		if (FRocketSupport::IsRocket())
-		{
-			// this was killing DDC build, for now we won't even try in rocket mode.
-			return NULL;
-		}
-
-		typedef void (*GetInterfaceVersionSimplygonSDKPtr)(ANSICHAR*);
-		typedef int (*InitializeSimplygonSDKPtr)(const char* LicenseData , SimplygonSDK::ISimplygonSDK** OutInterfacePtr);
- 
-		GetInterfaceVersionSimplygonSDKPtr GetInterfaceVersionSimplygonSDK;
-		InitializeSimplygonSDKPtr InitializeSimplygonSDK;
 		FString DllPath(FPaths::Combine(*FPaths::EngineDir(), TEXT("Binaries/ThirdParty/NotForLicensees/Simplygon")));
 		FString DllFilename(FPaths::Combine(*DllPath, TEXT("SimplygonSDKEpicUE4Releasex64.dll")));
- 
-		void* DLLHandle = FPlatformProcess::GetDllHandle(*DllFilename);
-		if (DLLHandle == NULL)
+		if( !FPaths::FileExists(DllFilename) )
 		{
 			DllFilename = FPaths::Combine(*DllPath, TEXT("SimplygonSDKRuntimeReleasex64.dll"));
-			DLLHandle = FPlatformProcess::GetDllHandle(*DllFilename);
 		}
- 
-		if( DLLHandle != NULL )
+
+		// If the DLL just doesn't exist, fail gracefully. Licensees and Rocket users will not necessarily have Simplygon.
+		if( !FPaths::FileExists(DllFilename) )
 		{
-			// Get API function pointers of interest
-			GetInterfaceVersionSimplygonSDK = (GetInterfaceVersionSimplygonSDKPtr)FPlatformProcess::GetDllExport( DLLHandle, TEXT( "GetInterfaceVersionSimplygonSDK" ) );
-			InitializeSimplygonSDK = (InitializeSimplygonSDKPtr)FPlatformProcess::GetDllExport( DLLHandle, TEXT( "InitializeSimplygonSDK" ) );
- 
-			if ((GetInterfaceVersionSimplygonSDK == NULL) || (InitializeSimplygonSDK == NULL))
-			{
-				// Couldn't find the functions we need.  
-				UE_LOG(LogSimplygon,Warning,TEXT("Failed to acquire Simplygon DLL exports."));
-				FPlatformProcess::FreeDllHandle( DLLHandle );
-				SDK = NULL;
-				return NULL;
-			}
+			UE_LOG(LogSimplygon,Warning,TEXT("Simplygon DLL not present - disabling."));
+			return NULL;
 		}
-		else
+
+		// Otherwise fail
+		void* DLLHandle = FPlatformProcess::GetDllHandle(*DllFilename);
+		if (DLLHandle == NULL)
 		{
 			int32 ErrorNum = FPlatformMisc::GetLastError();
 			TCHAR ErrorMsg[1024];
 			FPlatformMisc::GetSystemErrorMessage( ErrorMsg, 1024, ErrorNum );
-			UE_LOG(LogSimplygon,Warning,TEXT("Failed to get Simplygon DLL handle.\n\nError: %s (%d)"),ErrorMsg,ErrorNum);
-			SDK = NULL;
+			UE_LOG(LogSimplygon,Error,TEXT("Failed to get Simplygon DLL handle: %s (%d)"),ErrorMsg,ErrorNum);
 			return NULL;
 		}
 
+		// Get API function pointers of interest
+		typedef void (*GetInterfaceVersionSimplygonSDKPtr)(ANSICHAR*);
+		GetInterfaceVersionSimplygonSDKPtr GetInterfaceVersionSimplygonSDK = (GetInterfaceVersionSimplygonSDKPtr)FPlatformProcess::GetDllExport( DLLHandle, TEXT( "GetInterfaceVersionSimplygonSDK" ) );
+
+		typedef int (*InitializeSimplygonSDKPtr)(const char* LicenseData , SimplygonSDK::ISimplygonSDK** OutInterfacePtr);
+		InitializeSimplygonSDKPtr InitializeSimplygonSDK = (InitializeSimplygonSDKPtr)FPlatformProcess::GetDllExport( DLLHandle, TEXT( "InitializeSimplygonSDK" ) );
+ 
+		if ((GetInterfaceVersionSimplygonSDK == NULL) || (InitializeSimplygonSDK == NULL))
+		{
+			// Couldn't find the functions we need.  
+			UE_LOG(LogSimplygon,Warning,TEXT("Failed to acquire Simplygon DLL exports."));
+			FPlatformProcess::FreeDllHandle( DLLHandle );
+			return NULL;
+		}
+
+		ANSICHAR VersionHash[200];
 		GetInterfaceVersionSimplygonSDK(VersionHash);
 		if (FCStringAnsi::Strcmp(VersionHash, SimplygonSDK::GetInterfaceVersionHash()) != 0)
 		{
@@ -384,6 +376,7 @@ public:
 			LicenseData = (const char*)LicenseFileContents.GetTypedData();
 		}
 
+		SimplygonSDK::ISimplygonSDK* SDK = NULL;
 		int32 Result = InitializeSimplygonSDK(LicenseData, &SDK);
 		if (Result != SimplygonSDK::SG_ERROR_NOERROR && Result != SimplygonSDK::SG_ERROR_ALREADYINITIALIZED)
 		{
