Fixed metallic, roughness, specular color data extraction from simplygon

diff --git a/Engine/Source/Developer/SimplygonMeshReduction/Private/SimplygonMeshReduction.cpp b/Engine/Source/Developer/SimplygonMeshReduction/Private/SimplygonMeshReduction.cpp
--- a/Engine/Source/Developer/SimplygonMeshReduction/Private/SimplygonMeshReduction.cpp
+++ b/Engine/Source/Developer/SimplygonMeshReduction/Private/SimplygonMeshReduction.cpp
@@ -4,6 +4,12 @@
 #include "RawMesh.h"
 #include "MeshUtilities.h"
 
+// Standard Simplygon channels have some issues with extracting color data back from simplification, 
+// so we use this workaround with user channels
+static const char* USER_MATERIAL_CHANNEL_METALLIC = "UserMetallic";
+static const char* USER_MATERIAL_CHANNEL_ROUGHNESS = "UserRoughness";
+static const char* USER_MATERIAL_CHANNEL_SPECULAR = "UserSpecular";
+
 
 #ifdef __clang__
 	// SimplygonSDK.h uses 'deprecated' pragma which Clang does not recognize
@@ -523,10 +529,10 @@ public:
 			//Create a new material for the proxy geometry
 			SimplygonSDK::spMaterial OutputMaterialLOD = SDK->CreateMaterial();
 
-			//Create Image data where the diffuse data is stored
-			SimplygonSDK::spImageData OutputDiffuseData = SDK->CreateImageData();
 			//Cast diffuse texture data
 			{
+				//Create Image data where the diffuse data is stored
+				SimplygonSDK::spImageData OutputDiffuseData = SDK->CreateImageData();
 				// Cast the data using a color caster
 				SimplygonSDK::spColorCaster cast = SDK->CreateColorCaster();
 				cast->SetColorType( SimplygonSDK::SG_MATERIAL_CHANNEL_DIFFUSE );
@@ -578,7 +584,7 @@ public:
 				SimplygonSDK::spImageData OutputMetallicData = SDK->CreateImageData();
 				// Cast the data using a color caster
 				SimplygonSDK::spColorCaster cast = SDK->CreateColorCaster();
-				cast->SetColorType( SimplygonSDK::SG_MATERIAL_CHANNEL_METALLIC );
+				cast->SetColorType( USER_MATERIAL_CHANNEL_METALLIC );
 				cast->SetSourceMaterials( OriginalMaterials );
 				cast->SetMappingImage( MappingImage );
 				cast->SetOutputChannels( 3 );
@@ -588,8 +594,9 @@ public:
 				cast->SetOutputImage(OutputMetallicData);
 				cast->CastMaterials(); // Fetch!
 
-				OutputMaterialLOD->SetLayeredTextureImage(SimplygonSDK::SG_MATERIAL_CHANNEL_METALLIC, 0, OutputMetallicData);
-				OutputMaterialLOD->SetLayeredTextureLevel(SimplygonSDK::SG_MATERIAL_CHANNEL_METALLIC, 0, 0);
+				OutputMaterialLOD->AddUserChannel(USER_MATERIAL_CHANNEL_METALLIC);
+				OutputMaterialLOD->SetLayeredTextureImage(USER_MATERIAL_CHANNEL_METALLIC, 0, OutputMetallicData);
+				OutputMaterialLOD->SetLayeredTextureLevel(USER_MATERIAL_CHANNEL_METALLIC, 0, 0);
 			}
 
 			// Roughness texture data
@@ -599,7 +606,7 @@ public:
 				SimplygonSDK::spImageData OutputRoughnessData = SDK->CreateImageData();
 				// Cast the data using a color caster
 				SimplygonSDK::spColorCaster cast = SDK->CreateColorCaster();
-				cast->SetColorType( SimplygonSDK::SG_MATERIAL_CHANNEL_ROUGHNESS );
+				cast->SetColorType( USER_MATERIAL_CHANNEL_ROUGHNESS );
 				cast->SetSourceMaterials( OriginalMaterials );
 				cast->SetMappingImage( MappingImage ); 
 				cast->SetOutputChannels( 3 );
@@ -609,8 +616,9 @@ public:
 				cast->SetOutputImage(OutputRoughnessData);
 				cast->CastMaterials(); // Fetch!
 
-				OutputMaterialLOD->SetLayeredTextureImage(SimplygonSDK::SG_MATERIAL_CHANNEL_ROUGHNESS, 0, OutputRoughnessData);
-				OutputMaterialLOD->SetLayeredTextureLevel(SimplygonSDK::SG_MATERIAL_CHANNEL_ROUGHNESS, 0, 0);
+				OutputMaterialLOD->AddUserChannel(USER_MATERIAL_CHANNEL_ROUGHNESS);
+				OutputMaterialLOD->SetLayeredTextureImage(USER_MATERIAL_CHANNEL_ROUGHNESS, 0, OutputRoughnessData);
+				OutputMaterialLOD->SetLayeredTextureLevel(USER_MATERIAL_CHANNEL_ROUGHNESS, 0, 0);
 			}
 
 			// Specular texture data
@@ -620,7 +628,7 @@ public:
 				SimplygonSDK::spImageData OutputSpecularData = SDK->CreateImageData();
 				// Cast the data using a color caster
 				SimplygonSDK::spColorCaster cast = SDK->CreateColorCaster();
-				cast->SetColorType( SimplygonSDK::SG_MATERIAL_CHANNEL_SPECULAR );
+				cast->SetColorType( USER_MATERIAL_CHANNEL_SPECULAR );
 				cast->SetSourceMaterials( OriginalMaterials );
 				cast->SetMappingImage( MappingImage ); 
 				cast->SetOutputChannels( 3 );
@@ -630,8 +638,9 @@ public:
 				cast->SetOutputImage(OutputSpecularData);
 				cast->CastMaterials(); // Fetch!
 
-				OutputMaterialLOD->SetLayeredTextureImage(SimplygonSDK::SG_MATERIAL_CHANNEL_SPECULAR, 0, OutputSpecularData);
-				OutputMaterialLOD->SetLayeredTextureLevel(SimplygonSDK::SG_MATERIAL_CHANNEL_SPECULAR, 0, 0);
+				OutputMaterialLOD->AddUserChannel(USER_MATERIAL_CHANNEL_SPECULAR);
+				OutputMaterialLOD->SetLayeredTextureImage(USER_MATERIAL_CHANNEL_SPECULAR, 0, OutputSpecularData);
+				OutputMaterialLOD->SetLayeredTextureLevel(USER_MATERIAL_CHANNEL_SPECULAR, 0, 0);
 			}
 						
 			//Create a new material table for the new materials
@@ -1544,7 +1553,7 @@ private:
 				Texel[0] = InSamples[TexelIndex].R;
 				Texel[1] = InSamples[TexelIndex].G;
 				Texel[2] = InSamples[TexelIndex].B;
-				Texel[3] = InSamples[TexelIndex].A;
+				Texel[3] = (SGMaterialChannelName == SimplygonSDK::SG_MATERIAL_CHANNEL_DIFFUSE ? InSamples[TexelIndex].A : FColor::White.A);
 			
 				ImageColors->SetTuple(TexelIndex, Texel);
 			}
@@ -1588,8 +1597,16 @@ private:
 			if(!SGMaterial->HasUserChannel(SimplygonSDK::SG_MATERIAL_CHANNEL_ROUGHNESS))
 				SGMaterial->AddUserChannel(SimplygonSDK::SG_MATERIAL_CHANNEL_ROUGHNESS);
 
-			if(!SGMaterial->HasUserChannel(SimplygonSDK::SG_MATERIAL_CHANNEL_SPECULAR))
-				SGMaterial->AddUserChannel(SimplygonSDK::SG_MATERIAL_CHANNEL_SPECULAR);
+			// We actually use these channels for metallic, roughness and specular
+			if(!SGMaterial->HasUserChannel(USER_MATERIAL_CHANNEL_METALLIC))
+				SGMaterial->AddUserChannel(USER_MATERIAL_CHANNEL_METALLIC);
+
+			if(!SGMaterial->HasUserChannel(USER_MATERIAL_CHANNEL_ROUGHNESS))
+				SGMaterial->AddUserChannel(USER_MATERIAL_CHANNEL_ROUGHNESS);
+
+			if(!SGMaterial->HasUserChannel(USER_MATERIAL_CHANNEL_SPECULAR))
+				SGMaterial->AddUserChannel(USER_MATERIAL_CHANNEL_SPECULAR);
+
 					
 			SGMaterial->SetName(TCHAR_TO_ANSI(*FString::Printf(TEXT("Material%d"), MaterialIndex)));
 			
@@ -1610,21 +1627,21 @@ private:
 			if (FlattenMaterial.MetallicSamples.Num())
 			{
 				OutCastProperties.bCastMetallic = true;
-				SetMaterialChannelData(FlattenMaterial.MetallicSamples, FlattenMaterial.MetallicSize, SGMaterial, SimplygonSDK::SG_MATERIAL_CHANNEL_METALLIC);
+				SetMaterialChannelData(FlattenMaterial.MetallicSamples, FlattenMaterial.MetallicSize, SGMaterial, USER_MATERIAL_CHANNEL_METALLIC);
 			}
 
 			// Roughness
 			if (FlattenMaterial.RoughnessSamples.Num())
 			{
 				OutCastProperties.bCastRoughness = true;
-				SetMaterialChannelData(FlattenMaterial.RoughnessSamples, FlattenMaterial.RoughnessSize, SGMaterial, SimplygonSDK::SG_MATERIAL_CHANNEL_ROUGHNESS);
+				SetMaterialChannelData(FlattenMaterial.RoughnessSamples, FlattenMaterial.RoughnessSize, SGMaterial, USER_MATERIAL_CHANNEL_ROUGHNESS);
 			}
 
 			// Specular
 			if (FlattenMaterial.SpecularSamples.Num())
 			{
 				OutCastProperties.bCastSpecular = true;
-				SetMaterialChannelData(FlattenMaterial.SpecularSamples, FlattenMaterial.SpecularSize, SGMaterial, SimplygonSDK::SG_MATERIAL_CHANNEL_SPECULAR);
+				SetMaterialChannelData(FlattenMaterial.SpecularSamples, FlattenMaterial.SpecularSize, SGMaterial, USER_MATERIAL_CHANNEL_SPECULAR);
 			}
 			
  			OutSGMaterialTable->AddMaterial(SGMaterial);
@@ -1673,13 +1690,13 @@ private:
 		GetMaterialChannelData(SGMaterial, SimplygonSDK::SG_MATERIAL_CHANNEL_NORMALS, OutMaterial.NormalSamples, OutMaterial.NormalSize);
 	
 		// Metallic
-		GetMaterialChannelData(SGMaterial, SimplygonSDK::SG_MATERIAL_CHANNEL_METALLIC, OutMaterial.MetallicSamples, OutMaterial.MetallicSize);
+		GetMaterialChannelData(SGMaterial, USER_MATERIAL_CHANNEL_METALLIC, OutMaterial.MetallicSamples, OutMaterial.MetallicSize);
 
 		// Roughness
-		GetMaterialChannelData(SGMaterial, SimplygonSDK::SG_MATERIAL_CHANNEL_ROUGHNESS, OutMaterial.RoughnessSamples, OutMaterial.RoughnessSize);
+		GetMaterialChannelData(SGMaterial, USER_MATERIAL_CHANNEL_ROUGHNESS, OutMaterial.RoughnessSamples, OutMaterial.RoughnessSize);
 
 		// Specular
-		GetMaterialChannelData(SGMaterial, SimplygonSDK::SG_MATERIAL_CHANNEL_SPECULAR, OutMaterial.SpecularSamples, OutMaterial.SpecularSize);
+		GetMaterialChannelData(SGMaterial, USER_MATERIAL_CHANNEL_SPECULAR, OutMaterial.SpecularSamples, OutMaterial.SpecularSize);
 	}
 };
 
diff --git a/Engine/Source/Editor/UnrealEd/Private/EditorExporters.cpp b/Engine/Source/Editor/UnrealEd/Private/EditorExporters.cpp
--- a/Engine/Source/Editor/UnrealEd/Private/EditorExporters.cpp
+++ b/Engine/Source/Editor/UnrealEd/Private/EditorExporters.cpp
@@ -844,7 +844,7 @@ public:
 				// Only return for Opaque and Masked...
 				if (BlendMode == BLEND_Opaque || BlendMode == BLEND_Masked)
 				{
-					return Compiler->ForceCast(MaterialInterface->CompileProperty(Compiler, PropertyToCompile),MCT_Float1,true,true);
+					return Compiler->ForceCast(MaterialInterface->CompileProperty(Compiler, PropertyToCompile),MCT_Float,true,true);
 				}
 				break;
 			case MP_Normal:
