Fix for incorrect normals on spline meshes with HLOD generation.

diff --git a/Engine/Source/Developer/SimplygonMeshReduction/Private/SimplygonMeshReduction.cpp b/Engine/Source/Developer/SimplygonMeshReduction/Private/SimplygonMeshReduction.cpp
--- a/Engine/Source/Developer/SimplygonMeshReduction/Private/SimplygonMeshReduction.cpp
+++ b/Engine/Source/Developer/SimplygonMeshReduction/Private/SimplygonMeshReduction.cpp
@@ -853,15 +853,23 @@ public:
 			USplineMeshComponent* SplineMeshComponent = Cast<USplineMeshComponent>(InMeshComponent);
 			if (SplineMeshComponent)
 			{
+				for (int32 iVert = 0; iVert < RawMesh->WedgeIndices.Num(); ++iVert)
+				{
+					uint32 Index = RawMesh->WedgeIndices[iVert];
+
+					float& AxisValue = USplineMeshComponent::GetAxisValue(RawMesh->VertexPositions[Index], SplineMeshComponent->ForwardAxis);
+					FTransform SliceTransform = SplineMeshComponent->CalcSliceTransform(AxisValue);
+					RawMesh->WedgeTangentX[iVert] = SliceTransform.TransformVector(RawMesh->WedgeTangentX[iVert]);
+					RawMesh->WedgeTangentY[iVert] = SliceTransform.TransformVector(RawMesh->WedgeTangentY[iVert]);
+					RawMesh->WedgeTangentZ[iVert] = SliceTransform.TransformVector(RawMesh->WedgeTangentZ[iVert]);
+				}
+
 				for (int32 iVert = 0; iVert < RawMesh->VertexPositions.Num(); ++iVert)
 				{
-					float& Z = USplineMeshComponent::GetAxisValue(RawMesh->VertexPositions[iVert], SplineMeshComponent->ForwardAxis);
-					FTransform SliceTransform = SplineMeshComponent->CalcSliceTransform(Z);
-					Z = 0.0f;
-					RawMesh->VertexPositions[iVert] = SliceTransform.TransformPosition(RawMesh->VertexPositions[iVert]);
-					//RawMesh->WedgeTangentX[iVert] = SliceTransform.TransformVector(RawMesh->WedgeTangentX[iVert]);
-					//RawMesh->WedgeTangentY[iVert] = SliceTransform.TransformVector(RawMesh->WedgeTangentY[iVert]);
-					//RawMesh->WedgeTangentZ[iVert] = SliceTransform.TransformVector(RawMesh->WedgeTangentZ[iVert]);
+					float& AxisValue = USplineMeshComponent::GetAxisValue(RawMesh->VertexPositions[iVert], SplineMeshComponent->ForwardAxis);
+					FTransform SliceTransform = SplineMeshComponent->CalcSliceTransform(AxisValue);
+					AxisValue = 0.0f;
+					RawMesh->VertexPositions[iVert] = SliceTransform.TransformPosition(RawMesh->VertexPositions[iVert]);					
 				}
 			}
 
