Copying //UE4/Dev-Sequencer to //UE4/Dev-Main (Source: //UE4/Dev-Sequencer @ 3057646)

diff --git a/Engine/Plugins/Media/WmfMedia/Source/WmfMedia/Private/WmfMediaModule.cpp b/Engine/Plugins/Media/WmfMedia/Source/WmfMedia/Private/WmfMediaModule.cpp
--- a/Engine/Plugins/Media/WmfMedia/Source/WmfMedia/Private/WmfMediaModule.cpp
+++ b/Engine/Plugins/Media/WmfMedia/Source/WmfMedia/Private/WmfMediaModule.cpp
@@ -1,10 +1,9 @@
 // Copyright 1998-2016 Epic Games, Inc. All Rights Reserved.
 
-#include "WmfMediaPrivatePCH.h"
-#include "IMediaModule.h"
-#include "IMediaPlayerFactory.h"
-#include "ModuleInterface.h"
-#include "ModuleManager.h"
+#include "WmfMediaPCH.h"
+#include "IWmfMediaModule.h"
+#include "WmfMediaPlayer.h"
+#include "WmfMediaSettings.h"
 
 #pragma comment(lib, "mf")
 #pragma comment(lib, "mfplat")
@@ -22,8 +21,7 @@ DEFINE_LOG_CATEGORY(LogWmfMedia);
  * Implements the WmfMedia module.
  */
 class FWmfMediaModule
-	: public IModuleInterface
-	, public IMediaPlayerFactory
+	: public IWmfMediaModule
 {
 public:
 
@@ -34,20 +32,25 @@ public:
 
 public:
 
-	// IModuleInterface interface
+	//~ IWmfMediaModule interface
 
-	virtual void StartupModule() override
+	virtual TSharedPtr<IMediaPlayer> CreatePlayer() override
 	{
-		// load required libraries
-		IMediaModule* MediaModule = FModuleManager::LoadModulePtr<IMediaModule>("Media");
-
-		if (MediaModule == nullptr)
+		if (!Initialized)
 		{
-			UE_LOG(LogWmfMedia, Log, TEXT("Failed to load Media module"));
-
-			return;
+			return nullptr;
 		}
 
+		return MakeShareable(new FWmfMediaPlayer());
+	}
+
+public:
+
+	//~ IModuleInterface interface
+
+	virtual void StartupModule() override
+	{
+		// load required libraries
 		if (!LoadRequiredLibraries())
 		{
 			UE_LOG(LogWmfMedia, Log, TEXT("Failed to load required Windows Media Foundation libraries"));
@@ -65,99 +68,32 @@ public:
 			return;
 		}
 
-		// initialize supported media formats
-		SupportedFileTypes.Add(TEXT("3g2"), LOCTEXT("Format3g2", "3G2 Multimedia Stream"));
-		SupportedFileTypes.Add(TEXT("3gp"), LOCTEXT("Format3gp", "3GP Video Stream"));
-		SupportedFileTypes.Add(TEXT("3gp2"), LOCTEXT("Format3gp2", "3GPP2 Multimedia File"));
-		SupportedFileTypes.Add(TEXT("3gpp"), LOCTEXT("Format3gpp", "3GPP Multimedia File"));
-		SupportedFileTypes.Add(TEXT("aac"), LOCTEXT("FormatAac", "MPEG-2 Advanced Audio Coding File"));
-		SupportedFileTypes.Add(TEXT("adts"), LOCTEXT("FormatAdts", "Audio Data Transport Stream"));
-		SupportedFileTypes.Add(TEXT("asf"), LOCTEXT("FormatAsf", "ASF Media File"));
-		SupportedFileTypes.Add(TEXT("avi"), LOCTEXT("FormatAvi", "Audio Video Interleave File"));
-		SupportedFileTypes.Add(TEXT("m4a"), LOCTEXT("FormatM4a", "Apple MPEG-4 Audio"));
-		SupportedFileTypes.Add(TEXT("m4v"), LOCTEXT("FormatM4v", "Apple MPEG-4 Video"));
-		SupportedFileTypes.Add(TEXT("mov"), LOCTEXT("FormatMov", "Apple QuickTime Movie"));
-		SupportedFileTypes.Add(TEXT("mp3"), LOCTEXT("FormatMp3", "MPEG-2 Audio"));
-		SupportedFileTypes.Add(TEXT("mp4"), LOCTEXT("FormatMp4", "MPEG-4 Movie"));
-		SupportedFileTypes.Add(TEXT("sami"), LOCTEXT("FormatSami", "Synchronized Accessible Media Interchange (SAMI) File"));
-		SupportedFileTypes.Add(TEXT("smi"), LOCTEXT("FormatSmi", "Synchronized Multimedia Integration (SMIL) File"));
-		SupportedFileTypes.Add(TEXT("wav"), LOCTEXT("FormatWav", "Wave Audio File"));
-		SupportedFileTypes.Add(TEXT("wma"), LOCTEXT("FormatWma", "Windows Media Audio"));
-		SupportedFileTypes.Add(TEXT("wmv"), LOCTEXT("FormatWmv", "Windows Media Video"));
-
-		// initialize supported URI schemes
-		SupportedUriSchemes.Add(TEXT("http://"));
-		SupportedUriSchemes.Add(TEXT("httpd://"));
-		SupportedUriSchemes.Add(TEXT("https://"));
-		SupportedUriSchemes.Add(TEXT("mms://"));
-		SupportedUriSchemes.Add(TEXT("rtsp://"));
-		SupportedUriSchemes.Add(TEXT("rtspt://"));
-		SupportedUriSchemes.Add(TEXT("rtspu://"));
-
-		// register factory
-		MediaModule->RegisterPlayerFactory(*this);
+#if WITH_EDITOR
+		// register settings
+		ISettingsModule* SettingsModule = FModuleManager::GetModulePtr<ISettingsModule>("Settings");
+
+		if (SettingsModule != nullptr)
+		{
+			ISettingsSectionPtr SettingsSection = SettingsModule->RegisterSettings("Project", "Plugins", "WmfMedia",
+				LOCTEXT("WmfMediaSettingsName", "WMF Media"),
+				LOCTEXT("WmfMediaSettingsDescription", "Configure the WMF Media plug-in."),
+				GetMutableDefault<UWmfMediaSettings>()
+			);
+		}
+#endif // WITH_EDITOR
 
 		Initialized = true;
 	}
 
 	virtual void ShutdownModule() override
 	{
-		if (!Initialized)
-		{
-			return;
-		}
-
 		Initialized = false;
 
-		// unregister video player factory
-		IMediaModule* MediaModule = FModuleManager::GetModulePtr<IMediaModule>("Media");
-
-		if (MediaModule != nullptr)
-		{
-			MediaModule->UnregisterPlayerFactory(*this);
-		}		
-
-		// shutdown Windows Media Foundation
-		MFShutdown();
-	}
-
-public:
-
-	// IMediaPlayerFactory interface
-
-	virtual TSharedPtr<IMediaPlayer> CreatePlayer() override
-	{
 		if (Initialized)
 		{
-			return MakeShareable(new FWmfMediaPlayer());
-		}
-
-		return nullptr;
-	}
-
-	virtual const FMediaFileTypes& GetSupportedFileTypes() const override
-	{
-		return SupportedFileTypes;
-	}
-
-	virtual bool SupportsUrl(const FString& Url) const override
-	{
-		const FString Extension = FPaths::GetExtension(Url);
-
-		if (!Extension.IsEmpty() && SupportedFileTypes.Contains(Extension))
-		{
-			return true;
+			// shutdown Windows Media Foundation
+			MFShutdown();
 		}
-
-		for (const FString& Scheme : SupportedUriSchemes)
-		{
-			if (Url.StartsWith(Scheme))
-			{
-				return true;
-			}
-		}
-
-		return false;
 	}
 
 protected:
@@ -204,12 +140,6 @@ private:
 
 	/** Whether the module has been initialized. */
 	bool Initialized;
-
-	/** The collection of supported media file types. */
-	FMediaFileTypes SupportedFileTypes;
-
-	/** The collection of supported URI schemes. */
-	TArray<FString> SupportedUriSchemes;
 };
 
 
diff --git a/Engine/Plugins/Media/WmfMedia/Source/WmfMedia/Private/WmfMediaPCH.h b/Engine/Plugins/Media/WmfMedia/Source/WmfMedia/Private/WmfMediaPCH.h
new file mode 100644
--- /dev/null
+++ b/Engine/Plugins/Media/WmfMedia/Source/WmfMedia/Private/WmfMediaPCH.h
@@ -0,0 +1,31 @@
+// Copyright 1998-2016 Epic Games, Inc. All Rights Reserved.
+
+#pragma once
+
+#if WITH_EDITOR
+	#include "Developer/Settings/Public/ISettingsModule.h"
+	#include "Developer/Settings/Public/ISettingsSection.h"
+#endif
+
+#include "Runtime/Core/Public/Core.h"
+#include "Runtime/Core/Public/Modules/ModuleManager.h"
+#include "Runtime/CoreUObject/Public/CoreUObject.h"
+#include "Runtime/Media/Public/IMediaAudioSink.h"
+#include "Runtime/Media/Public/IMediaModule.h"
+#include "Runtime/Media/Public/IMediaOptions.h"
+#include "Runtime/Media/Public/IMediaStringSink.h"
+#include "Runtime/Media/Public/IMediaTextureSink.h"
+
+#include "Runtime/Core/Public/Windows/AllowWindowsPlatformTypes.h"
+#include <windows.h>
+#include <mfapi.h>
+#include <mferror.h>
+#include <mfidl.h>
+#include <nserror.h>
+#include <shlwapi.h>
+#include "Runtime/Core/Public/Windows/COMPointer.h"
+#include "Runtime/Core/Public/Windows/HideWindowsPlatformTypes.h"
+
+
+/** Log category for the WmfMedia module. */
+DECLARE_LOG_CATEGORY_EXTERN(LogWmfMedia, Log, All);
diff --git a/Engine/Plugins/Media/WmfMedia/Source/WmfMedia/Public/IWmfMediaModule.h b/Engine/Plugins/Media/WmfMedia/Source/WmfMedia/Public/IWmfMediaModule.h
new file mode 100644
--- /dev/null
+++ b/Engine/Plugins/Media/WmfMedia/Source/WmfMedia/Public/IWmfMediaModule.h
@@ -0,0 +1,30 @@
+// Copyright 1998-2016 Epic Games, Inc. All Rights Reserved.
+
+#pragma once
+
+#include "ModuleInterface.h"
+
+
+class IMediaPlayer;
+
+
+/**
+ * Interface for the WmfMedia module.
+ */
+class IWmfMediaModule
+	: public IModuleInterface
+{
+public:
+
+	/**
+	 * Creates a Windows Media Foundation based media player.
+	 *
+	 * @return A new media player, or nullptr if a player couldn't be created.
+	 */
+	virtual TSharedPtr<IMediaPlayer> CreatePlayer() = 0;
+
+public:
+
+	/** Virtual destructor. */
+	virtual ~IWmfMediaModule() { }
+};
diff --git a/Engine/Plugins/Media/WmfMedia/Source/WmfMedia/WmfMedia.Build.cs b/Engine/Plugins/Media/WmfMedia/Source/WmfMedia/WmfMedia.Build.cs
--- a/Engine/Plugins/Media/WmfMedia/Source/WmfMedia/WmfMedia.Build.cs
+++ b/Engine/Plugins/Media/WmfMedia/Source/WmfMedia/WmfMedia.Build.cs
@@ -9,19 +9,22 @@ namespace UnrealBuildTool.Rules
             DynamicallyLoadedModuleNames.AddRange(
                 new string[] {
                     "Media",
+					"Settings",
 				}
             );
 
 			PrivateDependencyModuleNames.AddRange(
 				new string[] {
 					"Core",
-                    "RenderCore",
+					"CoreUObject",
+					"RenderCore",
 				}
 			);
 
 			PrivateIncludePathModuleNames.AddRange(
 				new string[] {
                     "Media",
+					"Settings",
 				}
 			);
 
@@ -29,18 +32,18 @@ namespace UnrealBuildTool.Rules
 				new string[] {
 					"WmfMedia/Private",
                     "WmfMedia/Private/Player",
-                    "WmfMedia/Private/Tracks",
-                    "WmfMedia/Private/Wmf",
+					"WmfMedia/Private/Shared",
+					"WmfMedia/Private/Wmf",
 				}
 			);
 
             if ((Target.Platform == UnrealTargetPlatform.Win64) || (Target.Platform == UnrealTargetPlatform.Win32))
             {
-                PublicDelayLoadDLLs.Add("shlwapi.dll");
                 PublicDelayLoadDLLs.Add("mf.dll");
                 PublicDelayLoadDLLs.Add("mfplat.dll");
                 PublicDelayLoadDLLs.Add("mfplay.dll");
                 PublicDelayLoadDLLs.Add("mfuuid.dll");
+                PublicDelayLoadDLLs.Add("shlwapi.dll");
             }
 		}
 	}
diff --git a/Engine/Plugins/Media/WmfMedia/Source/WmfMediaEditor/Private/WmfMediaEditorModule.cpp b/Engine/Plugins/Media/WmfMedia/Source/WmfMediaEditor/Private/WmfMediaEditorModule.cpp
new file mode 100644
--- /dev/null
+++ b/Engine/Plugins/Media/WmfMedia/Source/WmfMediaEditor/Private/WmfMediaEditorModule.cpp
@@ -0,0 +1,22 @@
+// Copyright 1998-2016 Epic Games, Inc. All Rights Reserved.
+
+#include "WmfMediaEditorPCH.h"
+#include "ModuleInterface.h"
+
+
+/**
+ * Implements the WmfMediaEditor module.
+ */
+class FWmfMediaEditorModule
+	: public IModuleInterface
+{
+public:
+
+	//~ IModuleInterface interface
+
+	virtual void StartupModule() override { }
+	virtual void ShutdownModule() override { }
+};
+
+
+IMPLEMENT_MODULE(FWmfMediaEditorModule, WmfMediaEditor);
diff --git a/Engine/Plugins/Media/WmfMedia/Source/WmfMediaEditor/WmfMediaEditor.Build.cs b/Engine/Plugins/Media/WmfMedia/Source/WmfMediaEditor/WmfMediaEditor.Build.cs
new file mode 100644
--- /dev/null
+++ b/Engine/Plugins/Media/WmfMedia/Source/WmfMediaEditor/WmfMediaEditor.Build.cs
@@ -0,0 +1,25 @@
+// Copyright 1998-2016 Epic Games, Inc. All Rights Reserved.
+
+namespace UnrealBuildTool.Rules
+{
+	public class WmfMediaEditor : ModuleRules
+	{
+		public WmfMediaEditor(TargetInfo Target)
+		{
+			PrivateDependencyModuleNames.AddRange(
+				new string[] {
+					"Core",
+					"CoreUObject",
+					"MediaAssets",
+					"UnrealEd",
+                }
+			);
+
+			PrivateIncludePaths.AddRange(
+				new string[] {
+					"WmfMediaEditor/Private",
+				}
+			);
+		}
+	}
+}
diff --git a/Engine/Plugins/Media/WmfMedia/Source/WmfMediaFactory/Private/WmfMediaFactoryPCH.h b/Engine/Plugins/Media/WmfMedia/Source/WmfMediaFactory/Private/WmfMediaFactoryPCH.h
new file mode 100644
--- /dev/null
+++ b/Engine/Plugins/Media/WmfMedia/Source/WmfMediaFactory/Private/WmfMediaFactoryPCH.h
@@ -0,0 +1,14 @@
+// Copyright 1998-2016 Epic Games, Inc. All Rights Reserved.
+
+#pragma once
+
+#include "Runtime/Core/Public/Core.h"
+#include "Runtime/Core/Public/Modules/ModuleManager.h"
+#include "Runtime/Media/Public/IMediaModule.h"
+#include "Runtime/Media/Public/IMediaOptions.h"
+
+#include "../../WmfMedia/Public/IWmfMediaModule.h"
+
+
+/** Log category for the WmfMediaFactory module. */
+DECLARE_LOG_CATEGORY_EXTERN(LogWmfMediaFactory, Log, All);
diff --git a/Engine/Plugins/Media/WmfMedia/Source/WmfMediaFactory/WmfMediaFactory.Build.cs b/Engine/Plugins/Media/WmfMedia/Source/WmfMediaFactory/WmfMediaFactory.Build.cs
new file mode 100644
--- /dev/null
+++ b/Engine/Plugins/Media/WmfMedia/Source/WmfMediaFactory/WmfMediaFactory.Build.cs
@@ -0,0 +1,42 @@
+// Copyright 1998-2016 Epic Games, Inc. All Rights Reserved.
+
+namespace UnrealBuildTool.Rules
+{
+	public class WmfMediaFactory : ModuleRules
+	{
+		public WmfMediaFactory(TargetInfo Target)
+		{
+            DynamicallyLoadedModuleNames.AddRange(
+                new string[] {
+                    "Media",
+				}
+            );
+
+			PrivateDependencyModuleNames.AddRange(
+				new string[] {
+					"Core",
+                    "CoreUObject",
+                    "MediaAssets",
+                }
+			);
+
+			PrivateIncludePathModuleNames.AddRange(
+				new string[] {
+					"Media",
+					"WmfMedia",
+				}
+			);
+
+			PrivateIncludePaths.AddRange(
+				new string[] {
+					"WmfMediaFactory/Private",
+				}
+			);
+
+			if ((Target.Platform == UnrealTargetPlatform.Win32) || (Target.Platform == UnrealTargetPlatform.Win64))
+			{
+				DynamicallyLoadedModuleNames.Add("WmfMedia");
+			}
+		}
+	}
+}
diff --git a/Engine/Plugins/Media/WmfMedia/WmfMedia.uplugin b/Engine/Plugins/Media/WmfMedia/WmfMedia.uplugin
--- a/Engine/Plugins/Media/WmfMedia/WmfMedia.uplugin
+++ b/Engine/Plugins/Media/WmfMedia/WmfMedia.uplugin
@@ -1,13 +1,14 @@
 {
 	"FileVersion" : 3,
-	"Version" : 1,
-	"VersionName" : "1.0",
+	"Version" : 2,
+	"VersionName" : "2.0",
 	"FriendlyName" : "WMF Media Player",
 	"Description" : "Implements a media player using the Windows Media Foundation framework.",
 	"Category" : "Media Players",
 	"CreatedBy" : "Epic Games, Inc.",
 	"CreatedByURL" : "http://epicgames.com",
 	"DocsURL" : "https://forums.unrealengine.com/showthread.php?46879-Media-Framework-Documentation-for-4-5-Preview",
+	"EngineVersion" : "4.13.0",
 	"MarketplaceURL" : "",
 	"SupportURL" : "",
 	"EnabledByDefault" : true,
@@ -20,11 +21,18 @@
 			"Name" : "WmfMedia",
 			"Type" : "RuntimeNoCommandlet",
 			"LoadingPhase" : "PreLoadingScreen",
-			"WhitelistPlatforms" :
-			[
-				"Win32",
-				"Win64"
-			]
+			"WhitelistPlatforms" : [ "Win32", "Win64" ]
+		},
+		{
+			"Name" : "WmfMediaEditor",
+			"Type" : "Editor",
+			"LoadingPhase" : "PostEngineInit"
+		},
+		{
+			"Name" : "WmfMediaFactory",
+			"Type" : "Runtime",
+			"LoadingPhase" : "PostEngineInit",
+			"WhitelistPlatforms" : [ "Mac", "Linux", "Win32", "Win64" ]
 		}
 	]
 }
\ No newline at end of file
diff --git a/Engine/Shaders/MediaShaders.usf b/Engine/Shaders/MediaShaders.usf
--- a/Engine/Shaders/MediaShaders.usf
+++ b/Engine/Shaders/MediaShaders.usf
@@ -1,48 +1,209 @@
 // Copyright 1998-2016 Epic Games, Inc. All Rights Reserved.
 
 /*=============================================================================
-MediaShaders.usf: Shaders to convert various media decoder outputs to RGBA
-=============================================================================*/
+ MediaShaders.usf: Shaders to convert various media decoder outputs to RGBA
+ =============================================================================*/
 
 #include "Common.usf"
+#include "GammaCorrectionCommon.usf"
 
-/*=============================================================================
-* Vertex shader
+
+/* Constants
 =============================================================================*/
 
-struct FMediaShaderVertexOutput
+static const float3x3 YUV_TO_RGB_BT601 =
+{
+	1.000000, 0.000000, 1.139830,
+	1.000000,-0.394650,-0.580600,
+	1.000000, 2.032110, 0.000000,
+};
+
+static const float3x3 YUV_TO_RGB_BT709 =
 {
-	float4 Position : SV_POSITION;
-	float2 UV : TEXCOORD0;
+	1.000000, 0.000000, 1.280330,
+	1.000000,-0.214820,-0.380590,
+	1.000000, 2.127980, 0.000000,
 };
 
+static const float3x3 YUV_TO_RGB_JPEG =
+{
+	1.000000, 0.000000, 1.402000,
+	1.000000,-0.344140,-0.714140,
+	1.000000, 1.772000, 0.000000,
+};
+
+static const float3x3 YUV_TO_RGB_MSDN =
+{
+	1.164383, 0.000000, 1.596027,
+	1.164383,-0.391762,-0.812968,
+	1.164383, 2.017232, 0.000000,
+};
+
+static const float3x3 YUV_TO_RGB_PS4 =
+{
+	1.164400, 0.000000, 1.792700,
+	1.164400,-0.213300,-0.532900,
+	1.164400, 2.112400, 0.000000,
+};
+
+
+/* Helpers
+=============================================================================*/
+
+float3 YuvToRgb(float3 YUV, float3x3 Coeff)
+{
+	YUV -= float3(0.0625, 0.5, 0.5);
+	float3 RGB = mul(Coeff, YUV);
+	return sRGBToLinear(RGB);
+}
+
+
+/* Vertex shader
+=============================================================================*/
+
 void MainVertexShader(
 	in float4 InPosition : ATTRIBUTE0,
 	in float2 InUV : ATTRIBUTE1,
-	out FMediaShaderVertexOutput Output
-	)
+	out float4 OutPosition : SV_POSITION,
+	out float2 OutUV : TEXCOORD0)
 {
-	Output.Position = InPosition;
-	Output.UV = InUV;
+	OutPosition = InPosition;
+	OutUV = InUV;
 }
 
-/*=============================================================================
-* Pixel shader
+
+/* Pixel shaders
 =============================================================================*/
 
-//shader to convert from Full-range YCbCr to RGBA
+// shader to convert a AYUV texture to RGBA
+void AYUVConvertPS(
+	float4 InPosition : SV_POSITION,
+	float2 InUV : TEXCOORD0,
+	out float4 OutColor : SV_Target0)
+{
+	float4 AYUV = AYUVConvertUB.AYUVTexture.Sample(AYUVConvertUB.AYUVSampler, InUV.xy);
+	
+	OutColor.xyz = YuvToRgb(AYUV.zyx, YUV_TO_RGB_MSDN);
+	OutColor.w = AYUV.w;
+}
+
+
+// shader to convert a NV12 texture to RGBA
+void NV12ConvertPS(
+	float4 InPosition : SV_POSITION,
+	float2 InUV : TEXCOORD0,
+	out float4 OutColor : SV_Target0)
+{
+	float3 YUV;
+	YUV.x = NV12ConvertUB.NV12Texture.Sample(NV12ConvertUB.NV12Sampler, float2(InUV.x, InUV.y / 3.0f * 2.0f)).x;
+
+	uint UVOffs = uint(InUV.x * NV12ConvertUB.Width / 2.0f) * 2;
+	float UPos = (float(UVOffs) + 0.5f) / NV12ConvertUB.Width;
+	float VPos = (float(UVOffs) + 1.5f) / NV12ConvertUB.Width;
+
+	YUV.y = NV12ConvertUB.NV12Texture.Sample(NV12ConvertUB.NV12Sampler, float2(UPos, (InUV.y + 2.0f) / 3.0f)).x;
+	YUV.z = NV12ConvertUB.NV12Texture.Sample(NV12ConvertUB.NV12Sampler, float2(VPos, (InUV.y + 2.0f) / 3.0f)).x;
+
+	OutColor.xyz = YuvToRgb(YUV, YUV_TO_RGB_MSDN);
+	OutColor.w = 1.0f;
+}
+
+
+// shader to convert a NV21 texture to RGBA
+void NV21ConvertPS(
+	float4 InPosition : SV_POSITION,
+	float2 InUV : TEXCOORD0,
+	out float4 OutColor : SV_Target0)
+{
+	float3 YUV;
+	YUV.x = NV12ConvertUB.NV12Texture.Sample(NV12ConvertUB.NV12Sampler, float2(InUV.x, InUV.y / 3.0f * 2.0f)).x;
+
+	uint VUPos = uint(InUV.x * NV12ConvertUB.Width / 2.0f) * 2;
+	float UPos = (float(VUPos) + 1.5f) / NV12ConvertUB.Width;
+	float VPos = (float(VUPos) + 0.5f) / NV12ConvertUB.Width;
+
+	YUV.y = NV12ConvertUB.NV12Texture.Sample(NV12ConvertUB.NV12Sampler, float2(UPos, (InUV.y + 2.0f) / 3.0f)).x;
+	YUV.z = NV12ConvertUB.NV12Texture.Sample(NV12ConvertUB.NV12Sampler, float2(VPos, (InUV.y + 2.0f) / 3.0f)).x;
+
+	OutColor.xyz = YuvToRgb(YUV, YUV_TO_RGB_MSDN);
+	OutColor.w = 1.0f;
+}
+
+
+// shader to convert a luma (Y) and chroma (CbCr) texture to RGBA
 void YCbCrConvertPS(
+	float4 InPosition : SV_POSITION,
 	float2 InUV : TEXCOORD0,
 	out float4 OutColor : SV_Target0)
 {
-	float3 YCbCr;
-	YCbCr.x = YCbCrConvertUB.LumaTexture.Sample(YCbCrConvertUB.LumaSampler, InUV.xy).x - 0.0625; //e.g. 16 / 255
-	YCbCr.yz = YCbCrConvertUB.CbCrTexture.Sample(YCbCrConvertUB.CbCrSampler, InUV.xy).xy - float2(0.5, 0.5);
-			
+	float3 YUV;
+	YUV.x = YCbCrConvertUB.LumaTexture.Sample(YCbCrConvertUB.LumaSampler, InUV.xy).x;
+	YUV.yz = YCbCrConvertUB.CbCrTexture.Sample(YCbCrConvertUB.CbCrSampler, InUV.xy).xy;
 	
-	OutColor =	float4(
-					dot(float3(1.1644f, 0.0f, 1.7927f), YCbCr), // R
-					dot(float3(1.1644f, -0.2133f, -0.5329f), YCbCr), // G
-					dot(float3(1.1644f, 2.1124f, 0.0f), YCbCr), // B
-					1.0f);
-}
\ No newline at end of file
+	OutColor.xyz = YuvToRgb(YUV, YUV_TO_RGB_PS4);
+	OutColor.w = 1.0f;
+}
+
+
+// shader to convert a UYVY texture to RGBA
+void UYVYConvertPS(
+	float4 InPosition : SV_POSITION,
+	float2 InUV : TEXCOORD0,
+	out float4 OutColor : SV_Target0)
+{
+	float4 UYVY = UYVYConvertUB.UYVYTexture.Sample(UYVYConvertUB.UYVYSampler, float2(InUV.x, InUV.y));
+
+	uint Pos = InUV.x * UYVYConvertUB.Width;
+	float3 YUV = (Pos % 2 == 0) ? UYVY.yxz : UYVY.wxz;
+
+	OutColor.xyz = YuvToRgb(YUV, YUV_TO_RGB_MSDN);
+	OutColor.w = 1.0f;
+}
+
+
+// shader to convert three Y, U, and V textures to RGBA
+void YUVConvertPS(
+	float4 InPosition : SV_POSITION,
+	float2 InUV : TEXCOORD0,
+	out float4 OutColor : SV_Target0)
+{
+	float3 YUV;
+	YUV.x = YUVConvertUB.YTexture.Sample(YUVConvertUB.YSampler, InUV.xy).x;
+	YUV.y = YUVConvertUB.UTexture.Sample(YUVConvertUB.USampler, InUV.xy).x;
+	YUV.z = YUVConvertUB.VTexture.Sample(YUVConvertUB.VSampler, InUV.xy).x;
+	
+	OutColor.xyz = YuvToRgb(YUV, YUV_TO_RGB_MSDN);
+	OutColor.w = 1.0f;
+}
+
+
+// shader to convert a YUY2 texture to RGBA
+void YUY2ConvertPS(
+	float4 InPosition : SV_POSITION,
+	float2 InUV : TEXCOORD0,
+	out float4 OutColor : SV_Target0)
+{
+	float4 YUY2 = YUY2ConvertUB.YUY2Texture.Sample(YUY2ConvertUB.YUY2Sampler, float2(InUV.x, InUV.y));
+
+	uint Pos = InUV.x * YUY2ConvertUB.Width;
+	float3 YUV = (Pos % 2 == 0) ? YUY2.xyw : YUY2.zyw;
+
+	OutColor.xyz = YuvToRgb(YUV, YUV_TO_RGB_MSDN);
+	OutColor.w = 1.0f;
+}
+
+
+// shader to convert a YVYU texture to RGBA
+void YVYUConvertPS(
+	float4 InPosition : SV_POSITION,
+	float2 InUV : TEXCOORD0,
+	out float4 OutColor : SV_Target0)
+{
+	float4 YVYU = YVYUConvertUB.YVYUTexture.Sample(YVYUConvertUB.YVYUSampler, float2(InUV.x, InUV.y));
+
+	uint Pos = InUV.x * YVYUConvertUB.Width;
+	float3 YUV = (Pos % 2 == 0) ? YVYU.xwy : YVYU.zwy;
+
+	OutColor.xyz = YuvToRgb(YUV, YUV_TO_RGB_MSDN);
+	OutColor.w = 1.0f;
+}
diff --git a/Engine/Source/Developer/AssetTools/AssetTools.Build.cs b/Engine/Source/Developer/AssetTools/AssetTools.Build.cs
--- a/Engine/Source/Developer/AssetTools/AssetTools.Build.cs
+++ b/Engine/Source/Developer/AssetTools/AssetTools.Build.cs
@@ -12,6 +12,11 @@ public class AssetTools : ModuleRules
 			new string[] {
 				"Core",
 				"CoreUObject",
+			}
+		);
+
+		PrivateDependencyModuleNames.AddRange(
+			new string[] {
                 "CurveAssetEditor",
 				"Engine",
                 "InputCore",
diff --git a/Engine/Source/Developer/AssetTools/Public/IAssetTools.h b/Engine/Source/Developer/AssetTools/Public/IAssetTools.h
--- a/Engine/Source/Developer/AssetTools/Public/IAssetTools.h
+++ b/Engine/Source/Developer/AssetTools/Public/IAssetTools.h
@@ -2,13 +2,16 @@
 
 #pragma once
 
+#include "CoreUObject.h"
 #include "AssetTypeCategories.h"
 
+
 class IAssetTypeActions;
 class IClassTypeActions;
 class FAssetData;
 class UFactory;
 
+
 struct FAssetRenameData
 {
 	TWeakObjectPtr<UObject> Asset;
@@ -28,8 +31,10 @@ struct FAssetRenameData
 	}
 };
 
+
 DECLARE_MULTICAST_DELEGATE_OneParam(FAssetPostRenameEvent, const TArray<FAssetRenameData>&);
 
+
 struct FAdvancedAssetCategory
 {
 	EAssetTypeCategories::Type CategoryType;
@@ -42,6 +47,7 @@ struct FAdvancedAssetCategory
 	}
 };
 
+
 class IAssetTools
 {
 public:
diff --git a/Engine/Source/Developer/SimplygonMeshReduction/Private/SimplygonMeshReduction.cpp b/Engine/Source/Developer/SimplygonMeshReduction/Private/SimplygonMeshReduction.cpp
--- a/Engine/Source/Developer/SimplygonMeshReduction/Private/SimplygonMeshReduction.cpp
+++ b/Engine/Source/Developer/SimplygonMeshReduction/Private/SimplygonMeshReduction.cpp
@@ -1449,9 +1449,9 @@ private:
 
 	/**
 	 * Creates a Simplygon scene representation of the skeletal hierarchy.
-	 * @param InScene - The Simplygon scene.
-	 * @param InSkeleton - The skeletal hierarchy from which to create the Simplygon representation.
-	 * @param OutBoneTableIDs - A mapping of Bone IDs from RefSkeleton to Simplygon Bone Table IDs
+	 * @param InScene The Simplygon scene.
+	 * @param InSkeleton The skeletal hierarchy from which to create the Simplygon representation.
+	 * @param OutBoneTableIDs A mapping of Bone IDs from RefSkeleton to Simplygon Bone Table IDs
 	 */
 	void CreateSkeletalHierarchy( SimplygonSDK::spScene& InScene, const FReferenceSkeleton& InSkeleton, TArray<SimplygonSDK::rid>& OutBoneTableIDs ) 
 	{
@@ -1500,9 +1500,9 @@ private:
 
 	/**
 	 * Creates a Simplygon geometry representation from a skeletal mesh LOD model.
-	 * @param Scene - The Simplygon scene.
-	 * @param LODModel - The skeletal mesh LOD model from which to create the geometry data.
-	 * @param BoneIDs - A maps of Bone IDs from RefSkeleton to Simplygon BoneTable IDs
+	 * @param Scene The Simplygon scene.
+	 * @param LODModel The skeletal mesh LOD model from which to create the geometry data.
+	 * @param BoneIDs A maps of Bone IDs from RefSkeleton to Simplygon BoneTable IDs
 	 * @returns a Simplygon geometry data representation of the skeletal mesh LOD.
 	 */
 	SimplygonSDK::spGeometryData CreateGeometryFromSkeletalLODModel( SimplygonSDK::spScene& Scene, const FStaticLODModel& LODModel, const TArray<SimplygonSDK::rid>& BoneIDs )
@@ -1694,8 +1694,8 @@ private:
 
 	/**
 	 * Extracts mesh data from the Simplygon geometry representation in to a set of data usable by skeletal meshes.
-	 * @param GeometryData - the Simplygon geometry.
-	 * @param MeshData - the skeletal mesh output data.
+	 * @param GeometryData the Simplygon geometry.
+	 * @param MeshData the skeletal mesh output data.
 	 */
 	void ExtractSkeletalDataFromGeometry( const SimplygonSDK::spGeometryData& GeometryData, FSkeletalMeshData& MeshData )
 	{
@@ -1925,9 +1925,9 @@ private:
 
 	/**
 	 * Creates a skeletal mesh LOD model from the Simplygon geometry representation.
-	 * @param GeometryData - the Simplygon geometry representation from which to create a skeletal mesh LOD.
-	 * @param SkeletalMesh - the skeletal mesh in to which the LOD model will be added.
-	 * @param NewModel - the LOD model in to which the geometry will be stored.
+	 * @param GeometryData the Simplygon geometry representation from which to create a skeletal mesh LOD.
+	 * @param SkeletalMesh the skeletal mesh in to which the LOD model will be added.
+	 * @param NewModel the LOD model in to which the geometry will be stored.
 	 */
 	void CreateSkeletalLODModelFromGeometry( const SimplygonSDK::spGeometryData& GeometryData, const FReferenceSkeleton& RefSkeleton, FStaticLODModel* NewModel )
 	{
@@ -1984,8 +1984,8 @@ private:
 
 	/**
 	 * Sets reduction settings for Simplygon.
-	 * @param Settings - The skeletal mesh optimization settings.
-	 * @param ReductionSettings - The reduction settings to set for Simplygon.
+	 * @param Settings The skeletal mesh optimization settings.
+	 * @param ReductionSettings The reduction settings to set for Simplygon.
 	 */
 	void SetReductionSettings( const FSkeletalMeshOptimizationSettings& Settings, float BoundsRadius, int32 SourceTriCount, SimplygonSDK::spReductionSettings ReductionSettings )
 	{
@@ -2028,8 +2028,8 @@ private:
 	
 	/**
 	 * Sets vertex normal settings for Simplygon.
-	 * @param Settings - The skeletal mesh optimization settings.
-	 * @param NormalSettings - The normal settings to set for Simplygon.
+	 * @param Settings The skeletal mesh optimization settings.
+	 * @param NormalSettings The normal settings to set for Simplygon.
 	 */
 	void SetNormalSettings( const FSkeletalMeshOptimizationSettings& Settings, SimplygonSDK::spNormalCalculationSettings NormalSettings )
 	{
@@ -2041,8 +2041,8 @@ private:
 
 	/**
 	 * Sets Bone Lod settings for Simplygon.
-	 * @param Settings - The skeletal mesh optimization settings.
-	 * @param NormalSettings - The Bone LOD to set for Simplygon.
+	 * @param Settings The skeletal mesh optimization settings.
+	 * @param NormalSettings The Bone LOD to set for Simplygon.
 	 */
 	void SetBoneSettings( const FSkeletalMeshOptimizationSettings& Settings, SimplygonSDK::spBoneSettings BoneSettings)
 	{
@@ -2053,7 +2053,7 @@ private:
 
 	/**
 	 * Calculates the view distance that a mesh should be displayed at.
-	 * @param MaxDeviation - The maximum surface-deviation between the reduced geometry and the original. This value should be acquired from Simplygon
+	 * @param MaxDeviation The maximum surface-deviation between the reduced geometry and the original. This value should be acquired from Simplygon
 	 * @returns The calculated view distance	 
 	 */
 	float CalculateViewDistance( float MaxDeviation )
@@ -2339,9 +2339,9 @@ private:
 
 	/**
 	 * The method converts ESimplygonTextureSamplingQuality
-	 * @param InSamplingQuality - The Caster Settings used to setup the Simplygon Caster.
-	 * @param InMappingImage - Simplygon MappingImage.
-	 * @result 
+	 * @param InSamplingQuality The Caster Settings used to setup the Simplygon Caster.
+	 * @param InMappingImage Simplygon MappingImage.
+	 * @return 
 	 */
 	uint8 GetSamples(ESimplygonTextureSamplingQuality::Type InSamplingQuality)
 	{
@@ -2376,11 +2376,11 @@ private:
 
 	/**
 	 * Use Simplygon Color Caster to Cast a Color Channel
-	 * @param InCasterSettings - The Caster Settings used to setup the Simplygon Caster.
-	 * @param InMappingImage - Simplygon MappingImage.
-	 * @param InMaterialTable - Input MaterialTable used by the Simplygon Caster.
-	 * @param InTextureTable - Simplygon TextureTable used by the Simplyogn Caster.
-	 * @param OutColorData - The Simplygon Output ImageData.
+	 * @param InCasterSettings The Caster Settings used to setup the Simplygon Caster.
+	 * @param InMappingImage Simplygon MappingImage.
+	 * @param InMaterialTable Input MaterialTable used by the Simplygon Caster.
+	 * @param InTextureTable Simplygon TextureTable used by the Simplyogn Caster.
+	 * @param OutColorData The Simplygon Output ImageData.
 	 */
 	void CastColorChannel(const FSimplygonChannelCastingSettings& InCasterSettings, 
 		SimplygonSDK::spMappingImage& InMappingImage,
@@ -2406,11 +2406,11 @@ private:
 
 	/**
 	 * Use Simplygon Normal Caster to Cast a Normals Channel
-	 * @param InCasterSettings - The Caster Settings used to setup the Simplygon Caster.
-	 * @param InMappingImage - Simplygon MappingImage.
-	 * @param InMaterialTable - Input MaterialTable used by the Simplygon Caster.
-	 * @param InTextureTable - Simplygon TextureTable used by the Simplyogn Caster.
-	 * @param OutColorData - The Simplygon Output ImageData.
+	 * @param InCasterSettings The Caster Settings used to setup the Simplygon Caster.
+	 * @param InMappingImage Simplygon MappingImage.
+	 * @param InMaterialTable Input MaterialTable used by the Simplygon Caster.
+	 * @param InTextureTable Simplygon TextureTable used by the Simplyogn Caster.
+	 * @param OutColorData The Simplygon Output ImageData.
 	 */
 	void CastNormalsChannel(FSimplygonChannelCastingSettings InCasterSettings, 
 		SimplygonSDK::spMappingImage& InMappingImage,
@@ -2437,11 +2437,11 @@ private:
 
 	/**
 	 * Use Simplygon Opacity Caster to Cast an Opacity Channel
-	 * @param InCasterSettings - The Caster Settings used to setup the Simplygon Caster.
-	 * @param InMappingImage - Simplygon MappingImage.
-	 * @param InMaterialTable - Input MaterialTable used by the Simplygon Caster.
-	 * @param InTextureTable - Simplygon TextureTable used by the Simplyogn Caster.
-	 * @param OutColorData - The Simplygon Output ImageData.
+	 * @param InCasterSettings The Caster Settings used to setup the Simplygon Caster.
+	 * @param InMappingImage Simplygon MappingImage.
+	 * @param InMaterialTable Input MaterialTable used by the Simplygon Caster.
+	 * @param InTextureTable Simplygon TextureTable used by the Simplyogn Caster.
+	 * @param OutColorData The Simplygon Output ImageData.
 	 */
 	void CastOpacityChannel(FSimplygonChannelCastingSettings InCasterSettings, 
 		SimplygonSDK::spMappingImage& InMappingImage,
@@ -2462,8 +2462,8 @@ private:
 
 	/**
 	* Sets Mapping Image Setting for Simplygon.
-	* @param InMaterialLODSettings - The MaterialLOD Settings used for setting up Simplygon MappingImageSetting.
-	* @param InMappingImageSettings - The Simplygon MappingImage Settings that is being setup.
+	* @param InMaterialLODSettings The MaterialLOD Settings used for setting up Simplygon MappingImageSetting.
+	* @param InMappingImageSettings The Simplygon MappingImage Settings that is being setup.
 	*/
 	void SetupMappingImage(const FSimplygonMaterialLODSettings& InMaterialLODSettings,
 		SimplygonSDK::spMappingImageSettings InMappingImageSettings,
diff --git a/Engine/Source/Editor/ContentBrowser/Private/AssetViewWidgets.cpp b/Engine/Source/Editor/ContentBrowser/Private/AssetViewWidgets.cpp
--- a/Engine/Source/Editor/ContentBrowser/Private/AssetViewWidgets.cpp
+++ b/Engine/Source/Editor/ContentBrowser/Private/AssetViewWidgets.cpp
@@ -577,26 +577,27 @@ void SAssetViewItem::DirtyStateChanged()
 
 FText SAssetViewItem::GetAssetClassText() const
 {
-	if (AssetItem.IsValid())
+	if (!AssetItem.IsValid())
 	{
-		if (AssetItem->GetType() != EAssetItemType::Folder)
-		{
-			if (AssetTypeActions.IsValid())
-			{
-				return AssetTypeActions.Pin()->GetName();
-			}
-			else
-			{
-				return FText::FromName(StaticCastSharedPtr<FAssetViewAsset>(AssetItem)->Data.AssetClass);
-			}
-		}
-		else
+		return FText();
+	}
+	
+	if (AssetItem->GetType() == EAssetItemType::Folder)
+	{
+		return LOCTEXT("FolderName", "Folder");
+	}
+
+	if (AssetTypeActions.IsValid())
+	{
+		FText Name = AssetTypeActions.Pin()->GetName();
+
+		if (!Name.IsEmpty())
 		{
-			return LOCTEXT("FolderName", "Folder");
+			return Name;
 		}
 	}
 
-	return FText();
+	return FText::FromName(StaticCastSharedPtr<FAssetViewAsset>(AssetItem)->Data.AssetClass);
 }
 
 const FSlateBrush* SAssetViewItem::GetSCCStateImage() const
diff --git a/Engine/Source/Editor/ContentBrowser/Private/SAssetView.cpp b/Engine/Source/Editor/ContentBrowser/Private/SAssetView.cpp
--- a/Engine/Source/Editor/ContentBrowser/Private/SAssetView.cpp
+++ b/Engine/Source/Editor/ContentBrowser/Private/SAssetView.cpp
@@ -4002,7 +4002,7 @@ EVisibility SAssetView::IsAssetShowWarningTextVisible() const
 
 FText SAssetView::GetAssetShowWarningText() const
 {
-	if (AssetShowWarningText.IsBound())
+	if (AssetShowWarningText.IsSet())
 	{
 		return AssetShowWarningText.Get();
 	}
diff --git a/Engine/Source/Editor/ContentBrowser/Public/ContentBrowserDelegates.h b/Engine/Source/Editor/ContentBrowser/Public/ContentBrowserDelegates.h
--- a/Engine/Source/Editor/ContentBrowser/Public/ContentBrowserDelegates.h
+++ b/Engine/Source/Editor/ContentBrowser/Public/ContentBrowserDelegates.h
@@ -4,14 +4,23 @@
 #include "Developer/AssetTools/Public/IAssetTypeActions.h"
 #include "Runtime/AssetRegistry/Public/ARFilter.h"
 
+
+class FAssetData;
+class FDragDropOperation;
+class FExtender;
+class FReply;
+class SToolTip;
+struct FCollectionNameType;
+
+
 /** Called when a collection is selected in the collections view */
-DECLARE_DELEGATE_OneParam( FOnCollectionSelected, const struct FCollectionNameType& /*SelectedCollection*/);
+DECLARE_DELEGATE_OneParam( FOnCollectionSelected, const FCollectionNameType& /*SelectedCollection*/);
 
 /** Called to retrieve the tooltip for the specified asset */
 DECLARE_DELEGATE_RetVal_OneParam( TSharedRef< SToolTip >, FConstructToolTipForAsset, const FAssetData& /*Asset*/);
 
 /** Called to check if an asset should be filtered out by external code. Return true to exclude the asset from the view. */
-DECLARE_DELEGATE_RetVal_OneParam(bool, FOnShouldFilterAsset, const class FAssetData& /*AssetData*/);
+DECLARE_DELEGATE_RetVal_OneParam(bool, FOnShouldFilterAsset, const FAssetData& /*AssetData*/);
 
 /** Called to check if an asset tag should be display in details view. Return false to exclude the asset from the view. */
 DECLARE_DELEGATE_RetVal_TwoParams(bool, FOnShouldDisplayAssetTag, FName /*AssetType*/, FName /*TagName*/);
@@ -38,19 +47,19 @@ DECLARE_DELEGATE_OneParam(FSetPathPickerPathsDelegate, const TArray<FString>& /*
 DECLARE_DELEGATE_OneParam( FAdjustSelectionDelegate, const int32 /*direction*/ );
 
 /** Called when an asset is selected in the asset view */
-DECLARE_DELEGATE_OneParam(FOnAssetSelected, const class FAssetData& /*AssetData*/);
+DECLARE_DELEGATE_OneParam(FOnAssetSelected, const FAssetData& /*AssetData*/);
 
 /** Called when the user double clicks, presses enter, or presses space on an asset */
 DECLARE_DELEGATE_TwoParams(FOnAssetsActivated, const TArray<FAssetData>& /*ActivatedAssets*/, EAssetTypeActivationMethod::Type /*ActivationMethod*/);
 
 /** Called when an asset has begun being dragged by the user */
-DECLARE_DELEGATE_RetVal_OneParam(FReply, FOnAssetDragged, const TArray<class FAssetData>& /*AssetData*/);
+DECLARE_DELEGATE_RetVal_OneParam(FReply, FOnAssetDragged, const TArray<FAssetData>& /*AssetData*/);
 
 /** Called when an asset is clicked on in the asset view */
-DECLARE_DELEGATE_OneParam( FOnAssetClicked, const class FAssetData& /*AssetData*/ );
+DECLARE_DELEGATE_OneParam( FOnAssetClicked, const FAssetData& /*AssetData*/ );
 
 /** Called when an asset is double clicked in the asset view */
-DECLARE_DELEGATE_OneParam(FOnAssetDoubleClicked, const class FAssetData& /*AssetData*/);
+DECLARE_DELEGATE_OneParam(FOnAssetDoubleClicked, const FAssetData& /*AssetData*/);
 
 /** Called when enter is pressed on an asset in the asset view */
 DECLARE_DELEGATE_OneParam(FOnAssetEnterPressed, const TArray<FAssetData>& /*SelectedAssets*/);
@@ -76,10 +85,10 @@ DECLARE_DELEGATE_RetVal_OneParam(TSharedRef<FExtender>, FContentBrowserMenuExten
 DECLARE_DELEGATE_RetVal_ThreeParams(TSharedPtr<SWidget>, FOnGetFolderContextMenu, const TArray<FString>& /*SelectedPaths*/, FContentBrowserMenuExtender_SelectedPaths /*MenuExtender*/, FOnCreateNewFolder /*CreationDelegate*/);
 
 /** Called to request a custom asset item tooltip */
-DECLARE_DELEGATE_RetVal_OneParam( TSharedRef<SToolTip>, FOnGetCustomAssetToolTip, class FAssetData& /*AssetData*/);
+DECLARE_DELEGATE_RetVal_OneParam( TSharedRef<SToolTip>, FOnGetCustomAssetToolTip, FAssetData& /*AssetData*/);
 
 /** Called when an asset item visualizes its tooltip */
-DECLARE_DELEGATE_RetVal_TwoParams(bool, FOnVisualizeAssetToolTip, const TSharedPtr<SWidget>& /*ToolTipContent*/, class FAssetData& /*AssetData*/);
+DECLARE_DELEGATE_RetVal_TwoParams(bool, FOnVisualizeAssetToolTip, const TSharedPtr<SWidget>& /*ToolTipContent*/, FAssetData& /*AssetData*/);
 
 /** Called from the Asset Dialog when a non-modal dialog is closed or cancelled */
 DECLARE_DELEGATE(FOnAssetDialogCancelled);
@@ -93,6 +102,7 @@ DECLARE_DELEGATE_OneParam(FOnAssetsChosenForOpen, const TArray<FAssetData>& /*Se
 /** Called from the Asset Dialog when an asset name is chosen in non-modal Save dialogs */
 DECLARE_DELEGATE_OneParam(FOnObjectPathChosenForSave, const FString& /*ObjectPath*/);
 
+
 /** Contains the delegates used to handle a custom drag-and-drop in the asset view */
 struct FAssetViewDragAndDropExtender
 {
@@ -102,8 +112,7 @@ struct FAssetViewDragAndDropExtender
 			: DragDropOp(MoveTemp(InDragDropOp))
 			, PackagePaths(InPackagePaths)
 			, Collections(InCollections)
-		{
-		}
+		{ }
 
 		TSharedPtr<FDragDropOperation> DragDropOp;
 		const TArray<FName>& PackagePaths;
@@ -118,22 +127,19 @@ struct FAssetViewDragAndDropExtender
 		: OnDropDelegate(MoveTemp(InOnDropDelegate))
 		, OnDragOverDelegate()
 		, OnDragLeaveDelegate()
-	{
-	}
+	{ }
 
 	FAssetViewDragAndDropExtender(FOnDropDelegate InOnDropDelegate, FOnDragOverDelegate InOnDragOverDelegate)
 		: OnDropDelegate(MoveTemp(InOnDropDelegate))
 		, OnDragOverDelegate(MoveTemp(InOnDragOverDelegate))
 		, OnDragLeaveDelegate()
-	{
-	}
+	{ }
 
 	FAssetViewDragAndDropExtender(FOnDropDelegate InOnDropDelegate, FOnDragOverDelegate InOnDragOverDelegate, FOnDragLeaveDelegate InOnDragLeaveDelegate)
 		: OnDropDelegate(MoveTemp(InOnDropDelegate))
 		, OnDragOverDelegate(MoveTemp(InOnDragOverDelegate))
 		, OnDragLeaveDelegate(MoveTemp(InOnDragLeaveDelegate))
-	{
-	}
+	{ }
 
 	FOnDropDelegate OnDropDelegate;
 	FOnDragOverDelegate OnDragOverDelegate;
diff --git a/Engine/Source/Editor/ContentBrowser/Public/IContentBrowserSingleton.h b/Engine/Source/Editor/ContentBrowser/Public/IContentBrowserSingleton.h
--- a/Engine/Source/Editor/ContentBrowser/Public/IContentBrowserSingleton.h
+++ b/Engine/Source/Editor/ContentBrowser/Public/IContentBrowserSingleton.h
@@ -7,14 +7,17 @@
 #include "Developer/CollectionManager/Public/CollectionManagerTypes.h"
 #include "IFilter.h"
 #include "FilterCollection.h"
+#include "ITypedTableView.h"
 #include "AssetThumbnail.h"
 #include "UnrealClient.h"
 
+
 typedef const FAssetData& FAssetFilterType;
 typedef TFilterCollection<FAssetFilterType> FAssetFilterCollectionType;
 
 class UFactory;
 
+
 /** The view modes used in SAssetView */
 namespace EAssetViewType
 {
@@ -28,6 +31,7 @@ namespace EAssetViewType
 	};
 }
 
+
 /** A struct containing details about how the content browser should behave */
 struct FContentBrowserConfig
 {
@@ -100,7 +104,7 @@ struct FContentBrowserConfig
 		, bCanShowRealTimeThumbnails(true)
 		, bCanShowDevelopersFolder(true)
 		, bCanShowLockButton(true)
-	{}
+	{ }
 };
 
 
diff --git a/Engine/Source/Editor/DetailCustomizations/DetailCustomizations.Build.cs b/Engine/Source/Editor/DetailCustomizations/DetailCustomizations.Build.cs
--- a/Engine/Source/Editor/DetailCustomizations/DetailCustomizations.Build.cs
+++ b/Engine/Source/Editor/DetailCustomizations/DetailCustomizations.Build.cs
@@ -39,7 +39,6 @@ public class DetailCustomizations : ModuleRules
                 "InternationalizationSettings",
 				"SourceCodeAccess",
 				"RHI",
-                "MediaAssets",
                 "HardwareTargeting",
 				"SharedSettingsWidgets",
                 "LightPropagationVolumeRuntime",
diff --git a/Engine/Source/Editor/EditorStyle/Private/SlateEditorStyle.cpp b/Engine/Source/Editor/EditorStyle/Private/SlateEditorStyle.cpp
--- a/Engine/Source/Editor/EditorStyle/Private/SlateEditorStyle.cpp
+++ b/Engine/Source/Editor/EditorStyle/Private/SlateEditorStyle.cpp
@@ -1835,6 +1835,8 @@ void FSlateEditorStyle::FStyle::SetupGeneralStyles()
 		);
 
 		Set( "Sequencer.FixedFont", TTF_FONT( "Fonts/DroidSansMono", 9 ) );
+
+		Set( "Sequencer.RecordSelectedActors", new IMAGE_BRUSH( "SequenceRecorder/icon_tab_SequenceRecorder_16x", Icon16x16 ) );
 	}
 
 
@@ -5601,6 +5603,7 @@ void FSlateEditorStyle::FStyle::SetupClassIconsAndThumbnails()
 			TEXT("EditorTutorial"),
 			TEXT("EnvQuery"),
 			TEXT("ExponentialHeightFog"),
+			TEXT("FileMediaSource"),
 			TEXT("Font"),
 			TEXT("ForceFeedbackEffect"),
 			TEXT("GameMode"),
@@ -5647,6 +5650,7 @@ void FSlateEditorStyle::FStyle::SetupClassIconsAndThumbnails()
 			TEXT("PhysicsThrusterComponent"),
 			TEXT("PhysicsVolume"),
 			TEXT("PlaneReflectionCapture"),
+			TEXT("PlatformMediaSource"),
 			TEXT("PlayerController"),
 			TEXT("PlayerStart"),
 			TEXT("PointLight"),
@@ -5683,6 +5687,7 @@ void FSlateEditorStyle::FStyle::SetupClassIconsAndThumbnails()
 			TEXT("SpotLightStationary"),
 			TEXT("SpringArmComponent"),
 			TEXT("StaticMesh"),
+			TEXT("StreamMediaSource"),
 			TEXT("SubsurfaceProfile"),
 			TEXT("TargetPoint"),
 			TEXT("TextRenderActor"),
diff --git a/Engine/Source/Editor/MovieSceneTools/Private/CurveKeyAreas/StringCurveKeyArea.cpp b/Engine/Source/Editor/MovieSceneTools/Private/CurveKeyAreas/StringCurveKeyArea.cpp
new file mode 100644
--- /dev/null
+++ b/Engine/Source/Editor/MovieSceneTools/Private/CurveKeyAreas/StringCurveKeyArea.cpp
@@ -0,0 +1,251 @@
+// Copyright 1998-2016 Epic Games, Inc. All Rights Reserved.
+
+#include "MovieSceneToolsPrivatePCH.h"
+#include "StringCurveKeyArea.h"
+#include "SStringCurveKeyEditor.h"
+#include "MovieSceneClipboard.h"
+#include "Curves/StringCurve.h"
+#include "SequencerClipboardReconciler.h"
+
+/* IKeyArea interface
+ *****************************************************************************/
+
+TArray<FKeyHandle> FStringCurveKeyArea::AddKeyUnique(float Time, EMovieSceneKeyInterpolation InKeyInterpolation, float TimeToCopyFrom)
+{
+	TArray<FKeyHandle> AddedKeyHandles;
+	FKeyHandle CurrentKeyHandle = Curve->FindKey(Time);
+
+	if (Curve->IsKeyHandleValid(CurrentKeyHandle) == false)
+	{
+		if (OwningSection->GetStartTime() > Time)
+		{
+			OwningSection->SetStartTime(Time);
+		}
+
+		if (OwningSection->GetEndTime() < Time)
+		{
+			OwningSection->SetEndTime(Time);
+		}
+
+		FString Value = Curve->Eval(Time, Curve->GetDefaultValue());
+
+		if (TimeToCopyFrom != FLT_MAX)
+		{
+			Value = Curve->Eval(TimeToCopyFrom, Curve->GetDefaultValue());
+		}
+		else if ( IntermediateValue.IsSet() )
+		{
+			Value = IntermediateValue.GetValue();
+		}
+
+		Curve->AddKey(Time, Value, CurrentKeyHandle);
+		AddedKeyHandles.Add(CurrentKeyHandle);
+	}
+	else if ( IntermediateValue.IsSet() )
+	{
+		FString Value = IntermediateValue.GetValue();
+		Curve->UpdateOrAddKey(Time,Value);
+	}
+
+	return AddedKeyHandles;
+}
+
+TOptional<FKeyHandle> FStringCurveKeyArea::DuplicateKey(FKeyHandle KeyToDuplicate)
+{
+	if (!Curve->IsKeyHandleValid(KeyToDuplicate))
+	{
+		return TOptional<FKeyHandle>();
+	}
+
+	FStringCurveKey ThisKey = Curve->GetKey(KeyToDuplicate);
+
+	FKeyHandle KeyHandle = Curve->AddKey(GetKeyTime(KeyToDuplicate), ThisKey.Value);
+	// Ensure the rest of the key properties are added
+	Curve->GetKey(KeyHandle) = ThisKey;
+
+	return KeyHandle;
+}
+
+bool FStringCurveKeyArea::CanCreateKeyEditor()
+{
+	return true;
+}
+
+
+TSharedRef<SWidget> FStringCurveKeyArea::CreateKeyEditor(ISequencer* Sequencer)
+{
+	return SNew(SStringCurveKeyEditor)
+		.Sequencer(Sequencer)
+		.OwningSection(OwningSection)
+		.Curve(Curve)
+		.OnValueChanged(this, &FStringCurveKeyArea::OnValueChanged)
+		.IntermediateValue_Lambda([this] {
+			return IntermediateValue;
+		});
+};
+
+
+void FStringCurveKeyArea::DeleteKey(FKeyHandle KeyHandle)
+{
+	if (Curve->IsKeyHandleValid(KeyHandle))
+	{
+		Curve->DeleteKey(KeyHandle);
+	}
+}
+
+
+TOptional<FLinearColor> FStringCurveKeyArea::GetColor()
+{
+	return Color;
+}
+
+
+ERichCurveExtrapolation FStringCurveKeyArea::GetExtrapolationMode(bool bPreInfinity) const
+{
+	return RCCE_None;
+}
+
+ERichCurveInterpMode FStringCurveKeyArea::GetKeyInterpMode(FKeyHandle KeyHandle) const
+{
+	return RCIM_None;
+}
+
+
+TSharedPtr<FStructOnScope> FStringCurveKeyArea::GetKeyStruct(FKeyHandle KeyHandle)
+{
+	return MakeShareable(new FStructOnScope(FStringCurveKey::StaticStruct(), (uint8*)&Curve->GetKey(KeyHandle)));
+}
+
+
+ERichCurveTangentMode FStringCurveKeyArea::GetKeyTangentMode(FKeyHandle KeyHandle) const
+{
+	return RCTM_None;
+}
+
+
+float FStringCurveKeyArea::GetKeyTime(FKeyHandle KeyHandle) const
+{
+	return Curve->GetKeyTime(KeyHandle);
+}
+
+
+UMovieSceneSection* FStringCurveKeyArea::GetOwningSection()
+{
+	return OwningSection;
+}
+
+
+FRichCurve* FStringCurveKeyArea::GetRichCurve()
+{
+	return nullptr;
+}
+
+
+TArray<FKeyHandle> FStringCurveKeyArea::GetUnsortedKeyHandles() const
+{
+	TArray<FKeyHandle> OutKeyHandles;
+
+	for (auto It(Curve->GetKeyHandleIterator()); It; ++It)
+	{
+		OutKeyHandles.Add(It.Key());
+	}
+
+	return OutKeyHandles;
+}
+
+
+FKeyHandle FStringCurveKeyArea::MoveKey(FKeyHandle KeyHandle, float DeltaPosition)
+{
+	return Curve->SetKeyTime(KeyHandle, Curve->GetKeyTime(KeyHandle) + DeltaPosition);
+}
+
+
+void FStringCurveKeyArea::SetExtrapolationMode(ERichCurveExtrapolation ExtrapMode, bool bPreInfinity)
+{
+	// do nothing
+}
+
+bool FStringCurveKeyArea::CanSetExtrapolationMode() const
+{
+	return false;
+}
+
+void FStringCurveKeyArea::SetKeyInterpMode(FKeyHandle KeyHandle, ERichCurveInterpMode InterpMode)
+{
+	// do nothing
+}
+
+
+void FStringCurveKeyArea::SetKeyTangentMode(FKeyHandle KeyHandle, ERichCurveTangentMode TangentMode)
+{
+	// do nothing
+}
+
+
+void FStringCurveKeyArea::SetKeyTime(FKeyHandle KeyHandle, float NewKeyTime) const
+{
+	Curve->SetKeyTime(KeyHandle, NewKeyTime);
+}
+
+void FStringCurveKeyArea::CopyKeys(FMovieSceneClipboardBuilder& ClipboardBuilder, const TFunctionRef<bool(FKeyHandle, const IKeyArea&)>& KeyMask) const
+{
+	const UMovieSceneSection* Section = const_cast<FStringCurveKeyArea*>(this)->GetOwningSection();
+	UMovieSceneTrack* Track = Section ? Section->GetTypedOuter<UMovieSceneTrack>() : nullptr;
+	if (!Track)
+	{
+		return;
+	}
+
+	FMovieSceneClipboardKeyTrack* KeyTrack = nullptr;
+
+	for (auto It(Curve->GetKeyHandleIterator()); It; ++It)
+	{
+		FKeyHandle Handle = It.Key();
+		if (KeyMask(Handle, *this))
+		{
+			FStringCurveKey Key = Curve->GetKey(Handle);
+			if (!KeyTrack)
+			{
+				KeyTrack = &ClipboardBuilder.FindOrAddKeyTrack<FString>(GetName(), *Track);
+			}
+
+			KeyTrack->AddKey(Key.Time, Key.Value);
+		}
+	}
+}
+
+void FStringCurveKeyArea::PasteKeys(const FMovieSceneClipboardKeyTrack& KeyTrack, const FMovieSceneClipboardEnvironment& SrcEnvironment, const FSequencerPasteEnvironment& DstEnvironment)
+{
+	float PasteAt = DstEnvironment.CardinalTime;
+
+	KeyTrack.IterateKeys([&](const FMovieSceneClipboardKey& Key){
+		UMovieSceneSection* Section = GetOwningSection();
+		if (!Section)
+		{
+			return true;
+		}
+
+		if (Section->TryModify())
+		{		
+			float Time = PasteAt + Key.GetTime();
+			if (Section->GetStartTime() > Time)
+			{
+				Section->SetStartTime(Time);
+			}
+			if (Section->GetEndTime() < Time)
+			{
+				Section->SetEndTime(Time);
+			}
+
+			FKeyHandle KeyHandle = Curve->UpdateOrAddKey(Time, Key.GetValue<FString>());
+			DstEnvironment.ReportPastedKey(KeyHandle, *this);
+		}
+			
+		return true;
+	});
+}
+
+void FStringCurveKeyArea::OnValueChanged(FString InValue)
+{
+	ClearIntermediateValue();
+}
diff --git a/Engine/Source/Editor/MovieSceneTools/Private/CurveKeyEditors/SEnumCurveKeyEditor.cpp b/Engine/Source/Editor/MovieSceneTools/Private/CurveKeyEditors/SEnumCurveKeyEditor.cpp
--- a/Engine/Source/Editor/MovieSceneTools/Private/CurveKeyEditors/SEnumCurveKeyEditor.cpp
+++ b/Engine/Source/Editor/MovieSceneTools/Private/CurveKeyEditors/SEnumCurveKeyEditor.cpp
@@ -1,10 +1,13 @@
 // Copyright 1998-2016 Epic Games, Inc. All Rights Reserved.
 
 #include "MovieSceneToolsPrivatePCH.h"
+#include "MovieSceneToolHelpers.h"
 #include "SEnumCurveKeyEditor.h"
 
+
 #define LOCTEXT_NAMESPACE "EnumCurveKeyEditor"
 
+
 void SEnumCurveKeyEditor::Construct(const FArguments& InArgs)
 {
 	Sequencer = InArgs._Sequencer;
@@ -18,31 +21,37 @@ void SEnumCurveKeyEditor::Construct(const FArguments& InArgs)
 			InArgs._Enum,
 			TAttribute<int32>::Create(TAttribute<int32>::FGetter::CreateSP(this, &SEnumCurveKeyEditor::OnGetCurrentValue)),
 			FOnEnumSelectionChanged::CreateSP(this, &SEnumCurveKeyEditor::OnComboSelectionChanged),
-			InArgs._IntermediateValue)
+			InArgs._IntermediateValue
+		)
 	];
 }
 
+
 int32 SEnumCurveKeyEditor::OnGetCurrentValue() const
 {
 	float CurrentTime = Sequencer->GetCurrentLocalTime(*Sequencer->GetFocusedMovieSceneSequence());
 	return Curve->Evaluate(CurrentTime);
 }
 
+
 void SEnumCurveKeyEditor::OnComboSelectionChanged(int32 InSelectedItem, ESelectInfo::Type SelectInfo)
 {
 	FScopedTransaction Transaction(LOCTEXT("SetEnumKey", "Set Enum Key Value"));
 	OwningSection->SetFlags(RF_Transactional);
+
 	if (OwningSection->TryModify())
 	{
 		float CurrentTime = Sequencer->GetCurrentLocalTime(*Sequencer->GetFocusedMovieSceneSequence());
 
 		bool bKeyWillBeAdded = Curve->IsKeyHandleValid(Curve->FindKey(CurrentTime)) == false;
+
 		if (bKeyWillBeAdded)
 		{
 			if (OwningSection->GetStartTime() > CurrentTime)
 			{
 				OwningSection->SetStartTime(CurrentTime);
 			}
+
 			if (OwningSection->GetEndTime() < CurrentTime)
 			{
 				OwningSection->SetEndTime(CurrentTime);
@@ -64,4 +73,5 @@ void SEnumCurveKeyEditor::OnComboSelectionChanged(int32 InSelectedItem, ESelectI
 	}
 }
 
-#undef LOCTEXT_NAMESPACE
\ No newline at end of file
+
+#undef LOCTEXT_NAMESPACE
diff --git a/Engine/Source/Editor/MovieSceneTools/Private/CurveKeyEditors/SStringCurveKeyEditor.cpp b/Engine/Source/Editor/MovieSceneTools/Private/CurveKeyEditors/SStringCurveKeyEditor.cpp
new file mode 100644
--- /dev/null
+++ b/Engine/Source/Editor/MovieSceneTools/Private/CurveKeyEditors/SStringCurveKeyEditor.cpp
@@ -0,0 +1,77 @@
+// Copyright 1998t-2016 Epic Games, Inc. All Rights Reserved.
+
+#include "MovieSceneToolsPrivatePCH.h"
+#include "MovieSceneStringSection.h"
+#include "SStringCurveKeyEditor.h"
+
+#define LOCTEXT_NAMESPACE "StringCurveKeyEditor"
+
+void SStringCurveKeyEditor::Construct(const FArguments& InArgs)
+{
+	Sequencer = InArgs._Sequencer;
+	OwningSection = InArgs._OwningSection;
+	Curve = InArgs._Curve;
+	OnValueChangedEvent = InArgs._OnValueChanged;
+	IntermediateValue = InArgs._IntermediateValue;
+	float CurrentTime = Sequencer->GetCurrentLocalTime(*Sequencer->GetFocusedMovieSceneSequence());
+	ChildSlot
+	[
+		SNew(SEditableText)
+		.SelectAllTextWhenFocused(true)
+		.Text(this, &SStringCurveKeyEditor::GetText)
+		.OnTextCommitted(this, &SStringCurveKeyEditor::OnTextCommitted)
+ 	];
+}
+
+FText SStringCurveKeyEditor::GetText() const
+{
+	if ( IntermediateValue.IsSet() && IntermediateValue.Get().IsSet() )
+	{
+		return FText::FromString(IntermediateValue.Get().GetValue());
+	}
+
+	float CurrentTime = Sequencer->GetCurrentLocalTime(*Sequencer->GetFocusedMovieSceneSequence());
+
+	UMovieSceneStringSection* StringSection = Cast<UMovieSceneStringSection>(OwningSection);
+	FString CurrentValue = StringSection ? StringSection->Eval(CurrentTime, Curve->GetDefaultValue()) : Curve->Eval(CurrentTime, Curve->GetDefaultValue());
+	
+	return FText::FromString(CurrentValue);
+}
+
+void SStringCurveKeyEditor::OnTextCommitted(const FText& InText, ETextCommit::Type InCommitType)
+{
+	FScopedTransaction Transaction(LOCTEXT("SetStringKey", "Set String Key Value"));
+	OwningSection->SetFlags(RF_Transactional);
+	if (OwningSection->TryModify())
+	{
+		float CurrentTime = Sequencer->GetCurrentLocalTime(*Sequencer->GetFocusedMovieSceneSequence());
+
+		bool bKeyWillBeAdded = Curve->IsKeyHandleValid(Curve->FindKey(CurrentTime)) == false;
+		if (bKeyWillBeAdded)
+		{
+			if (OwningSection->GetStartTime() > CurrentTime)
+			{
+				OwningSection->SetStartTime(CurrentTime);
+			}
+			if (OwningSection->GetEndTime() < CurrentTime)
+			{
+				OwningSection->SetEndTime(CurrentTime);
+			}
+		}
+		
+		if (Curve->GetNumKeys() == 0)
+		{
+			Curve->SetDefaultValue(InText.ToString());
+		}
+		else
+		{
+			Curve->UpdateOrAddKey(CurrentTime, InText.ToString());
+		}
+
+		OnValueChangedEvent.ExecuteIfBound(InText.ToString());
+
+		Sequencer->NotifyMovieSceneDataChanged( EMovieSceneDataChangeType::TrackValueChanged );
+	}
+}
+
+#undef LOCTEXT_NAMESPACE
\ No newline at end of file
diff --git a/Engine/Source/Editor/MovieSceneTools/Private/CurveKeyEditors/SStringCurveKeyEditor.h b/Engine/Source/Editor/MovieSceneTools/Private/CurveKeyEditors/SStringCurveKeyEditor.h
new file mode 100644
--- /dev/null
+++ b/Engine/Source/Editor/MovieSceneTools/Private/CurveKeyEditors/SStringCurveKeyEditor.h
@@ -0,0 +1,34 @@
+// Copyright 1998-2016 Epic Games, Inc. All Rights Reserved.
+
+#pragma once
+
+/**
+ * A widget for editing a curve representing string keys.
+ */
+class SStringCurveKeyEditor : public SCompoundWidget
+{
+public:
+	/** Notification for string value change */
+	DECLARE_DELEGATE_OneParam( FOnValueChanged, FString );
+
+	SLATE_BEGIN_ARGS(SStringCurveKeyEditor) {}
+		SLATE_ARGUMENT(ISequencer*, Sequencer)
+		SLATE_ARGUMENT(UMovieSceneSection*, OwningSection)
+		SLATE_ARGUMENT(FStringCurve*, Curve)
+		SLATE_EVENT(FOnValueChanged, OnValueChanged)
+		SLATE_ATTRIBUTE(TOptional<FString>, IntermediateValue)
+	SLATE_END_ARGS();
+
+	void Construct(const FArguments& InArgs);
+
+private:
+
+	FText GetText() const;
+	void OnTextCommitted(const FText& InText, ETextCommit::Type InCommitType);
+
+	ISequencer* Sequencer;
+	UMovieSceneSection* OwningSection;
+	FStringCurve* Curve;
+	FOnValueChanged OnValueChangedEvent;
+	TAttribute<TOptional<FString>> IntermediateValue;
+};
\ No newline at end of file
diff --git a/Engine/Source/Editor/MovieSceneTools/Private/MovieSceneToolHelpers.cpp b/Engine/Source/Editor/MovieSceneTools/Private/MovieSceneToolHelpers.cpp
--- a/Engine/Source/Editor/MovieSceneTools/Private/MovieSceneToolHelpers.cpp
+++ b/Engine/Source/Editor/MovieSceneTools/Private/MovieSceneToolHelpers.cpp
@@ -20,6 +20,9 @@
 #include "MovieSceneCaptureHelpers.h"
 #include "FbxImporter.h"
 #include "MatineeImportTools.h"
+#include "MovieSceneToolHelpers.h"
+#include "MovieSceneToolsProjectSettings.h"
+
 
 /* MovieSceneToolHelpers
  *****************************************************************************/
diff --git a/Engine/Source/Editor/MovieSceneTools/Private/MovieSceneToolsModule.cpp b/Engine/Source/Editor/MovieSceneTools/Private/MovieSceneToolsModule.cpp
--- a/Engine/Source/Editor/MovieSceneTools/Private/MovieSceneToolsModule.cpp
+++ b/Engine/Source/Editor/MovieSceneTools/Private/MovieSceneToolsModule.cpp
@@ -18,6 +18,7 @@
 #include "VectorPropertyTrackEditor.h"
 #include "VisibilityPropertyTrackEditor.h"
 #include "ActorReferencePropertyTrackEditor.h"
+#include "StringPropertyTrackEditor.h"
 
 #include "TransformTrackEditor.h"
 #include "CameraCutTrackEditor.h"
@@ -44,9 +45,13 @@
 #include "Curves/CurveBase.h"
 #include "ISettingsModule.h"
 #include "PropertyEditorModule.h"
+#include "IMovieSceneTools.h"
+#include "MovieSceneToolsProjectSettings.h"
+
 
 #define LOCTEXT_NAMESPACE "FMovieSceneToolsModule"
 
+
 /**
  * Implements the MovieSceneTools module.
  */
@@ -78,6 +83,7 @@ public:
 		VectorPropertyTrackCreateEditorHandle = SequencerModule.RegisterTrackEditor_Handle( FOnCreateTrackEditor::CreateStatic( &FVectorPropertyTrackEditor::CreateTrackEditor ) );
 		VisibilityPropertyTrackCreateEditorHandle = SequencerModule.RegisterTrackEditor_Handle( FOnCreateTrackEditor::CreateStatic( &FVisibilityPropertyTrackEditor::CreateTrackEditor ) );
 		ActorReferencePropertyTrackCreateEditorHandle = SequencerModule.RegisterTrackEditor_Handle(FOnCreateTrackEditor::CreateStatic( &FActorReferencePropertyTrackEditor::CreateTrackEditor ) );
+		StringPropertyTrackCreateEditorHandle = SequencerModule.RegisterTrackEditor_Handle(FOnCreateTrackEditor::CreateStatic( &FStringPropertyTrackEditor::CreateTrackEditor ) );
 
 		// register specialty track editors
 		AnimationTrackCreateEditorHandle = SequencerModule.RegisterTrackEditor_Handle( FOnCreateTrackEditor::CreateStatic( &FSkeletalAnimationTrackEditor::CreateTrackEditor ) );
@@ -128,6 +134,7 @@ public:
 		SequencerModule.UnRegisterTrackEditor_Handle( VectorPropertyTrackCreateEditorHandle );
 		SequencerModule.UnRegisterTrackEditor_Handle( VisibilityPropertyTrackCreateEditorHandle );
 		SequencerModule.UnRegisterTrackEditor_Handle( ActorReferencePropertyTrackCreateEditorHandle );
+		SequencerModule.UnRegisterTrackEditor_Handle( StringPropertyTrackCreateEditorHandle );
 
 		// unregister specialty track editors
 		SequencerModule.UnRegisterTrackEditor_Handle( AnimationTrackCreateEditorHandle );
@@ -199,6 +206,7 @@ private:
 	FDelegateHandle VectorPropertyTrackCreateEditorHandle;
 	FDelegateHandle VisibilityPropertyTrackCreateEditorHandle;
 	FDelegateHandle ActorReferencePropertyTrackCreateEditorHandle;
+	FDelegateHandle StringPropertyTrackCreateEditorHandle;
 
 	FDelegateHandle AnimationTrackCreateEditorHandle;
 	FDelegateHandle AttachTrackCreateEditorHandle;
diff --git a/Engine/Source/Editor/MovieSceneTools/Private/MovieSceneToolsPrivatePCH.h b/Engine/Source/Editor/MovieSceneTools/Private/MovieSceneToolsPrivatePCH.h
--- a/Engine/Source/Editor/MovieSceneTools/Private/MovieSceneToolsPrivatePCH.h
+++ b/Engine/Source/Editor/MovieSceneTools/Private/MovieSceneToolsPrivatePCH.h
@@ -2,27 +2,21 @@
 
 #pragma once
 
+#include "Runtime/Core/Public/Core.h"
+#include "Runtime/CoreUObject/Public/CoreUObject.h"
+#include "Runtime/Engine/Classes/Camera/CameraComponent.h"
+#include "Runtime/Engine/Classes/Curves/IntegralCurve.h"
+#include "Runtime/Engine/Classes/Curves/NameCurve.h"
+#include "Runtime/Engine/Classes/Curves/RichCurve.h"
+#include "Runtime/MovieScene/Public/MovieScene.h"
+#include "Runtime/MovieScene/Public/MovieSceneSequenceInstance.h"
+#include "Runtime/MovieScene/Public/KeyParams.h"
+#include "Runtime/Slate/Public/SlateBasics.h"
 
-#include "UnrealEd.h"
-#include "MovieScene.h"
-#include "ScopedTransaction.h"
-#include "ISequencerSection.h"
-#include "ISectionLayoutBuilder.h"
-#include "KeyParams.h"
-#include "MovieSceneSequenceInstance.h"
-#include "MovieSceneTrackEditor.h"
-#include "IMovieSceneTools.h"
-#include "MovieSceneSection.h"
-#include "PropertySection.h"
-#include "PropertyTrackEditor.h"
-#include "CommonMovieSceneTools.h"
-#include "MovieSceneToolHelpers.h"
-#include "MovieSceneCommonHelpers.h"
-#include "MovieSceneToolsProjectSettings.h"
-
-#include "ThumbnailSection.h"
-#include "CameraCutSection.h"
-#include "CinematicShotSection.h"
-
-#include "TrackEditorThumbnail.h"
-#include "TrackEditorThumbnailPool.h"
+#include "Editor/Sequencer/Public/ISectionLayoutBuilder.h"
+#include "Editor/Sequencer/Public/ISequencerSection.h"
+#include "Editor/Sequencer/Public/MovieSceneTrackEditor.h"
+#include "Runtime/MovieScene/Public/MovieSceneCommonHelpers.h"
+#include "Runtime/MovieScene/Public/MovieSceneSection.h"
+#include "Editor/UnrealEd/Public/ScopedTransaction.h"
+#include "Editor/UnrealEd/Public/UnrealEd.h"
diff --git a/Engine/Source/Editor/MovieSceneTools/Private/Sections/ActorReferencePropertySection.h b/Engine/Source/Editor/MovieSceneTools/Private/Sections/ActorReferencePropertySection.h
--- a/Engine/Source/Editor/MovieSceneTools/Private/Sections/ActorReferencePropertySection.h
+++ b/Engine/Source/Editor/MovieSceneTools/Private/Sections/ActorReferencePropertySection.h
@@ -2,26 +2,31 @@
 
 #pragma once
 
-#include "ByteKeyArea.h"
+#include "IntegralKeyArea.h"
+#include "PropertySection.h"
+
+
+struct FIntegralCurve;
+class FActorReferenceKeyArea;
+class UMovieSceneSection;
 
 
 /** A key area for displaying and editing an actor reference property. */
 class FActorReferenceKeyArea
-	: public FIntegralKeyArea < int32 >
+	: public FIntegralKeyArea<int32>
 {
 public:
 
-	FActorReferenceKeyArea( FIntegralCurve& InCurve, UMovieSceneSection* InOwningSection )
-		: FIntegralKeyArea( InCurve, InOwningSection )
+	FActorReferenceKeyArea(FIntegralCurve& InCurve, UMovieSceneSection* InOwningSection)
+		: FIntegralKeyArea(InCurve, InOwningSection)
 	{ }
 
 public:
 
-	// FIntegralKeyArea overrides
+	//~ FIntegralKeyArea overrides
 
 	virtual bool CanCreateKeyEditor() override { return false; }
-	virtual TSharedRef<SWidget> CreateKeyEditor( ISequencer* Sequencer ) override { return SNew( SSpacer ); }
-
+	virtual TSharedRef<SWidget> CreateKeyEditor(ISequencer* Sequencer) override { return SNew(SSpacer); }
 };
 
 
@@ -39,7 +44,7 @@ public:
 
 public:
 
-	// FPropertySection interface 
+	//~ FPropertySection interface
 
 	virtual void GenerateSectionLayout(class ISectionLayoutBuilder& LayoutBuilder) const override;
 	virtual void SetIntermediateValue(FPropertyChangedParams PropertyChangedParams) override;
diff --git a/Engine/Source/Editor/MovieSceneTools/Private/Sections/BoolPropertySection.cpp b/Engine/Source/Editor/MovieSceneTools/Private/Sections/BoolPropertySection.cpp
--- a/Engine/Source/Editor/MovieSceneTools/Private/Sections/BoolPropertySection.cpp
+++ b/Engine/Source/Editor/MovieSceneTools/Private/Sections/BoolPropertySection.cpp
@@ -1,9 +1,12 @@
 // Copyright 1998-2016 Epic Games, Inc. All Rights Reserved.
 
 #include "MovieSceneToolsPrivatePCH.h"
+#include "BoolKeyArea.h"
+#include "BoolPropertySection.h"
+#include "CommonMovieSceneTools.h"
 #include "MovieSceneBoolTrack.h"
 #include "MovieSceneBoolSection.h"
-#include "BoolPropertySection.h"
+
 
 void FBoolPropertySection::GenerateSectionLayout( class ISectionLayoutBuilder& LayoutBuilder ) const
 {
diff --git a/Engine/Source/Editor/MovieSceneTools/Private/Sections/BoolPropertySection.h b/Engine/Source/Editor/MovieSceneTools/Private/Sections/BoolPropertySection.h
--- a/Engine/Source/Editor/MovieSceneTools/Private/Sections/BoolPropertySection.h
+++ b/Engine/Source/Editor/MovieSceneTools/Private/Sections/BoolPropertySection.h
@@ -2,28 +2,31 @@
 
 #pragma once
 
-#include "BoolKeyArea.h"
+#include "PropertySection.h"
+
+
+class FBoolKeyArea;
 
 
 /**
-* An implementation of bool property sections
+* An implementation of bool property sections.
 */
 class FBoolPropertySection
 	: public FPropertySection
 {
 public:
 
-	FBoolPropertySection( UMovieSceneSection& InSectionObject, const FText& SectionName, ISequencer* InSequencer )
-		: FPropertySection( InSectionObject, SectionName )
-		, Sequencer( InSequencer ) {}
+	FBoolPropertySection(UMovieSceneSection& InSectionObject, const FText& SectionName, ISequencer* InSequencer)
+		: FPropertySection(InSectionObject, SectionName)
+		, Sequencer(InSequencer) {}
 
 public:
 
-	// FPropertySection interface
+	//~ FPropertySection interface
 
-	virtual void GenerateSectionLayout( class ISectionLayoutBuilder& LayoutBuilder ) const override;
-	virtual int32 OnPaintSection( FSequencerSectionPainter& Painter ) const override;
-	virtual void SetIntermediateValue( FPropertyChangedParams PropertyChangedParams ) override;
+	virtual void GenerateSectionLayout(class ISectionLayoutBuilder& LayoutBuilder) const override;
+	virtual int32 OnPaintSection(FSequencerSectionPainter& Painter) const override;
+	virtual void SetIntermediateValue(FPropertyChangedParams PropertyChangedParams) override;
 	virtual void ClearIntermediateValue() override;
 
 protected:
diff --git a/Engine/Source/Editor/MovieSceneTools/Private/Sections/BytePropertySection.h b/Engine/Source/Editor/MovieSceneTools/Private/Sections/BytePropertySection.h
--- a/Engine/Source/Editor/MovieSceneTools/Private/Sections/BytePropertySection.h
+++ b/Engine/Source/Editor/MovieSceneTools/Private/Sections/BytePropertySection.h
@@ -2,7 +2,11 @@
 
 #pragma once
 
-#include "ByteKeyArea.h"
+#include "PropertySection.h"
+
+
+class FByteKeyArea;
+class UMovieSceneSection;
 
 
 /**
diff --git a/Engine/Source/Editor/MovieSceneTools/Private/Sections/CameraCutSection.cpp b/Engine/Source/Editor/MovieSceneTools/Private/Sections/CameraCutSection.cpp
--- a/Engine/Source/Editor/MovieSceneTools/Private/Sections/CameraCutSection.cpp
+++ b/Engine/Source/Editor/MovieSceneTools/Private/Sections/CameraCutSection.cpp
@@ -1,12 +1,14 @@
 // Copyright 1998-2016 Epic Games, Inc. All Rights Reserved.
 
 #include "MovieSceneToolsPrivatePCH.h"
+#include "CameraCutSection.h"
 #include "ISectionLayoutBuilder.h"
 #include "Runtime/MovieSceneTracks/Public/Sections/MovieSceneCameraCutSection.h"
 #include "Runtime/Engine/Public/Slate/SceneViewport.h"
 #include "SInlineEditableTextBlock.h"
 #include "MovieSceneToolsUserSettings.h"
 
+
 #define LOCTEXT_NAMESPACE "FCameraCutSection"
 
 
diff --git a/Engine/Source/Editor/MovieSceneTools/Private/Sections/CameraCutSection.h b/Engine/Source/Editor/MovieSceneTools/Private/Sections/CameraCutSection.h
--- a/Engine/Source/Editor/MovieSceneTools/Private/Sections/CameraCutSection.h
+++ b/Engine/Source/Editor/MovieSceneTools/Private/Sections/CameraCutSection.h
@@ -2,6 +2,9 @@
 
 #pragma once
 
+#include "ThumbnailSection.h"
+
+
 class FTrackEditorThumbnail;
 class FTrackEditorThumbnailPool;
 class IMenu;
diff --git a/Engine/Source/Editor/MovieSceneTools/Private/Sections/CinematicShotSection.cpp b/Engine/Source/Editor/MovieSceneTools/Private/Sections/CinematicShotSection.cpp
--- a/Engine/Source/Editor/MovieSceneTools/Private/Sections/CinematicShotSection.cpp
+++ b/Engine/Source/Editor/MovieSceneTools/Private/Sections/CinematicShotSection.cpp
@@ -1,6 +1,7 @@
 // Copyright 1998-2016 Epic Games, Inc. All Rights Reserved.
 
 #include "MovieSceneToolsPrivatePCH.h"
+#include "CinematicShotSection.h"
 #include "ISectionLayoutBuilder.h"
 #include "Runtime/MovieSceneTracks/Public/Sections/MovieSceneCinematicShotSection.h"
 #include "CinematicShotTrackEditor.h"
@@ -9,8 +10,13 @@
 #include "MovieSceneToolHelpers.h"
 #include "MovieSceneToolsUserSettings.h"
 
+
 #define LOCTEXT_NAMESPACE "FCinematicShotSection"
 
+
+/* FCinematicShotSection structors
+ *****************************************************************************/
+
 FCinematicShotSection::FCinematicSectionCache::FCinematicSectionCache(UMovieSceneCinematicShotSection* Section)
 	: ActualStartTime(0.f), TimeScale(0.f)
 {
@@ -22,9 +28,6 @@ FCinematicShotSection::FCinematicSectionCache::FCinematicSectionCache(UMovieScen
 }
 
 
-/* FCinematicShotSection structors
- *****************************************************************************/
-
 FCinematicShotSection::FCinematicShotSection(TSharedPtr<ISequencer> InSequencer, TSharedPtr<FTrackEditorThumbnailPool> InThumbnailPool, UMovieSceneSection& InSection, TSharedPtr<FCinematicShotTrackEditor> InCinematicShotTrackEditor)
 	: FThumbnailSection(InSequencer, InThumbnailPool, InSection)
 	, SectionObject(*CastChecked<UMovieSceneCinematicShotSection>(&InSection))
diff --git a/Engine/Source/Editor/MovieSceneTools/Private/Sections/CinematicShotSection.h b/Engine/Source/Editor/MovieSceneTools/Private/Sections/CinematicShotSection.h
--- a/Engine/Source/Editor/MovieSceneTools/Private/Sections/CinematicShotSection.h
+++ b/Engine/Source/Editor/MovieSceneTools/Private/Sections/CinematicShotSection.h
@@ -2,6 +2,9 @@
 
 #pragma once
 
+#include "ThumbnailSection.h"
+
+
 class FThumbnailSection;
 class FTrackEditorThumbnail;
 class FTrackEditorThumbnailPool;
diff --git a/Engine/Source/Editor/MovieSceneTools/Private/Sections/ColorPropertySection.cpp b/Engine/Source/Editor/MovieSceneTools/Private/Sections/ColorPropertySection.cpp
--- a/Engine/Source/Editor/MovieSceneTools/Private/Sections/ColorPropertySection.cpp
+++ b/Engine/Source/Editor/MovieSceneTools/Private/Sections/ColorPropertySection.cpp
@@ -4,6 +4,7 @@
 #include "MovieSceneColorTrack.h"
 #include "MovieSceneColorSection.h"
 #include "ColorPropertySection.h"
+#include "CommonMovieSceneTools.h"
 #include "MovieSceneSequence.h"
 #include "FloatCurveKeyArea.h"
 
@@ -23,6 +24,7 @@ void FColorPropertySection::GenerateSectionLayout( class ISectionLayoutBuilder&
 	LayoutBuilder.AddKeyArea( "A", NSLOCTEXT( "FColorPropertySection", "OpacityArea", "Opacity" ), AlphaKeyArea.ToSharedRef() );
 }
 
+
 int32 FColorPropertySection::OnPaintSection( FSequencerSectionPainter& Painter ) const
 {
 	int32 LayerId = Painter.PaintSectionBackground();
@@ -207,6 +209,7 @@ FLinearColor FColorPropertySection::FindSlateColor(const FName& ColorName) const
 	return FLinearColor(0.0f, 0.0f, 0.0f, 0.0f);
 }
 
+
 void FColorPropertySection::SetIntermediateValue( FPropertyChangedParams PropertyChangedParams )
 {
 	const UStructProperty* StructProp = Cast<const UStructProperty>( PropertyChangedParams.PropertyPath.Last() );
@@ -233,6 +236,7 @@ void FColorPropertySection::SetIntermediateValue( FPropertyChangedParams Propert
 	AlphaKeyArea->SetIntermediateValue(ColorValue.A);
 }
 
+
 void FColorPropertySection::ClearIntermediateValue()
 {
 	RedKeyArea->ClearIntermediateValue();
diff --git a/Engine/Source/Editor/MovieSceneTools/Private/Sections/ColorPropertySection.h b/Engine/Source/Editor/MovieSceneTools/Private/Sections/ColorPropertySection.h
--- a/Engine/Source/Editor/MovieSceneTools/Private/Sections/ColorPropertySection.h
+++ b/Engine/Source/Editor/MovieSceneTools/Private/Sections/ColorPropertySection.h
@@ -3,9 +3,15 @@
 #pragma once
 
 #include "FloatCurveKeyArea.h"
+#include "PropertySection.h"
 
 
+class FFloatCurveKeyArea;
+class ISequencer;
 class UMovieSceneColorSection;
+class UMovieSceneColorTrack;
+class UMovieSceneSection;
+class UMovieSceneTrack;
 
 
 /**
diff --git a/Engine/Source/Editor/MovieSceneTools/Private/Sections/FloatPropertySection.cpp b/Engine/Source/Editor/MovieSceneTools/Private/Sections/FloatPropertySection.cpp
--- a/Engine/Source/Editor/MovieSceneTools/Private/Sections/FloatPropertySection.cpp
+++ b/Engine/Source/Editor/MovieSceneTools/Private/Sections/FloatPropertySection.cpp
@@ -1,22 +1,26 @@
 // Copyright 1998-2016 Epic Games, Inc. All Rights Reserved.
 
 #include "MovieSceneToolsPrivatePCH.h"
+#include "FloatCurveKeyArea.h"
 #include "FloatPropertySection.h"
 #include "MovieSceneFloatSection.h"
 
-void FFloatPropertySection::GenerateSectionLayout( class ISectionLayoutBuilder& LayoutBuilder ) const
+
+void FFloatPropertySection::GenerateSectionLayout(class ISectionLayoutBuilder& LayoutBuilder) const
 {
-	UMovieSceneFloatSection* FloatSection = Cast<UMovieSceneFloatSection>( &SectionObject );
-	KeyArea = MakeShareable( new FFloatCurveKeyArea( &FloatSection->GetFloatCurve(), FloatSection ) );
-	LayoutBuilder.SetSectionAsKeyArea( KeyArea.ToSharedRef() );
+	UMovieSceneFloatSection* FloatSection = Cast<UMovieSceneFloatSection>(&SectionObject);
+	KeyArea = MakeShareable(new FFloatCurveKeyArea(&FloatSection->GetFloatCurve(), FloatSection));
+	LayoutBuilder.SetSectionAsKeyArea(KeyArea.ToSharedRef());
 }
 
-void FFloatPropertySection::SetIntermediateValue( FPropertyChangedParams PropertyChangedParams )
+
+void FFloatPropertySection::SetIntermediateValue(FPropertyChangedParams PropertyChangedParams)
 {
-	KeyArea->SetIntermediateValue( PropertyChangedParams.GetPropertyValue<float>() );
+	KeyArea->SetIntermediateValue(PropertyChangedParams.GetPropertyValue<float>());
 }
 
+
 void FFloatPropertySection::ClearIntermediateValue()
 {
 	KeyArea->ClearIntermediateValue();
-}
\ No newline at end of file
+}
diff --git a/Engine/Source/Editor/MovieSceneTools/Private/Sections/FloatPropertySection.h b/Engine/Source/Editor/MovieSceneTools/Private/Sections/FloatPropertySection.h
--- a/Engine/Source/Editor/MovieSceneTools/Private/Sections/FloatPropertySection.h
+++ b/Engine/Source/Editor/MovieSceneTools/Private/Sections/FloatPropertySection.h
@@ -2,11 +2,15 @@
 
 #pragma once
 
-#include "FloatCurveKeyArea.h"
+#include "PropertySection.h"
+
+
+class FFloatCurveKeyArea;
+class UMovieSceneSection;
 
 
 /**
- * An implementation of float property sections
+ * An implementation of float property sections.
  */
 class FFloatPropertySection
 	: public FPropertySection
@@ -19,7 +23,7 @@ public:
 
 public:
 
-	// FPropertySection interface 
+	//~ FPropertySection interface 
 
 	virtual void GenerateSectionLayout(class ISectionLayoutBuilder& LayoutBuilder) const override;
 	virtual void SetIntermediateValue(FPropertyChangedParams PropertyChangedParams) override;
diff --git a/Engine/Source/Editor/MovieSceneTools/Private/Sections/ParameterSection.h b/Engine/Source/Editor/MovieSceneTools/Private/Sections/ParameterSection.h
--- a/Engine/Source/Editor/MovieSceneTools/Private/Sections/ParameterSection.h
+++ b/Engine/Source/Editor/MovieSceneTools/Private/Sections/ParameterSection.h
@@ -2,6 +2,11 @@
 
 #pragma once
 
+#include "PropertySection.h"
+
+
+class UMovieSceneSection;
+
 
 /**
  * A movie scene section for material parameters.
@@ -17,13 +22,14 @@ public:
 
 public:
 
-	// FPropertySection interface
+	//~ FPropertySection interface
+
 	virtual void SetIntermediateValue(FPropertyChangedParams PropertyChangedParams) override { }
 	virtual void ClearIntermediateValue() override { }
 
 public:
 
-	// ISequencerSection interface
+	//~ ISequencerSection interface
 
 	virtual void GenerateSectionLayout(class ISectionLayoutBuilder& LayoutBuilder) const override;
 	virtual bool RequestDeleteCategory(const TArray<FName>& CategoryNamePath) override;
diff --git a/Engine/Source/Editor/MovieSceneTools/Private/Sections/StringPropertySection.cpp b/Engine/Source/Editor/MovieSceneTools/Private/Sections/StringPropertySection.cpp
new file mode 100644
--- /dev/null
+++ b/Engine/Source/Editor/MovieSceneTools/Private/Sections/StringPropertySection.cpp
@@ -0,0 +1,38 @@
+// Copyright 1998-2016 Epic Games, Inc. All Rights Reserved.
+
+#include "MovieSceneToolsPrivatePCH.h"
+#include "StringCurveKeyArea.h"
+#include "StringPropertySection.h"
+#include "MovieSceneStringSection.h"
+
+
+void FStringPropertySection::GenerateSectionLayout(class ISectionLayoutBuilder& LayoutBuilder) const
+{
+	UMovieSceneStringSection* StringSection = Cast<UMovieSceneStringSection>(&SectionObject);
+	KeyArea = MakeShareable(new FStringCurveKeyArea(&StringSection->GetStringCurve(), StringSection));
+	LayoutBuilder.SetSectionAsKeyArea(KeyArea.ToSharedRef());
+}
+
+
+void FStringPropertySection::SetIntermediateValue(FPropertyChangedParams PropertyChangedParams)
+{
+	void* CurrentObject = PropertyChangedParams.ObjectsThatChanged[0];
+	void* PropertyValue = nullptr;
+	for (int32 i = 0; i < PropertyChangedParams.PropertyPath.Num(); i++)
+	{
+		CurrentObject = PropertyChangedParams.PropertyPath[i]->ContainerPtrToValuePtr<FString>(CurrentObject, 0);
+	}
+
+	const UStrProperty* StrProperty = Cast<const UStrProperty>( PropertyChangedParams.PropertyPath.Last() );
+	if ( StrProperty )
+	{
+		FString StrPropertyValue = StrProperty->GetPropertyValue(CurrentObject);
+		KeyArea->SetIntermediateValue(StrPropertyValue);
+	}
+}
+
+
+void FStringPropertySection::ClearIntermediateValue()
+{
+	KeyArea->ClearIntermediateValue();
+}
diff --git a/Engine/Source/Editor/MovieSceneTools/Private/Sections/StringPropertySection.h b/Engine/Source/Editor/MovieSceneTools/Private/Sections/StringPropertySection.h
new file mode 100644
--- /dev/null
+++ b/Engine/Source/Editor/MovieSceneTools/Private/Sections/StringPropertySection.h
@@ -0,0 +1,35 @@
+// Copyright 1998-2016 Epic Games, Inc. All Rights Reserved.
+
+#pragma once
+
+#include "PropertySection.h"
+
+
+class FStringCurveKeyArea;
+class UMovieSceneSection;
+
+
+/**
+ * An implementation of FString property sections.
+ */
+class FStringPropertySection
+	: public FPropertySection
+{
+public:
+
+	FStringPropertySection(UMovieSceneSection& InSectionObject, const FText& SectionName)
+		: FPropertySection(InSectionObject, SectionName)
+	{ }
+
+public:
+
+	//~ FPropertySection interface 
+
+	virtual void GenerateSectionLayout(class ISectionLayoutBuilder& LayoutBuilder) const override;
+	virtual void SetIntermediateValue(FPropertyChangedParams PropertyChangedParams) override;
+	virtual void ClearIntermediateValue() override;
+
+private:
+
+	mutable TSharedPtr<FStringCurveKeyArea> KeyArea;
+};
diff --git a/Engine/Source/Editor/MovieSceneTools/Private/Sections/ThumbnailSection.cpp b/Engine/Source/Editor/MovieSceneTools/Private/Sections/ThumbnailSection.cpp
--- a/Engine/Source/Editor/MovieSceneTools/Private/Sections/ThumbnailSection.cpp
+++ b/Engine/Source/Editor/MovieSceneTools/Private/Sections/ThumbnailSection.cpp
@@ -6,6 +6,8 @@
 #include "PropertyEditorModule.h"
 #include "IDetailsView.h"
 #include "MovieSceneToolsUserSettings.h"
+#include "ThumbnailSection.h"
+
 
 #define LOCTEXT_NAMESPACE "FThumbnailSection"
 
@@ -30,16 +32,19 @@ FThumbnailSection::FThumbnailSection(TSharedPtr<ISequencer> InSequencer, TShared
 	GetMutableDefault<UMovieSceneUserThumbnailSettings>()->OnForceRedraw().AddRaw(this, &FThumbnailSection::RedrawThumbnails);
 }
 
+
 FThumbnailSection::~FThumbnailSection()
 {
 	GetMutableDefault<UMovieSceneUserThumbnailSettings>()->OnForceRedraw().RemoveAll(this);
 }
 
+
 void FThumbnailSection::RedrawThumbnails()
 {
 	ThumbnailCache.ForceRedraw();
 }
 
+
 /* FThumbnailSection interface
  *****************************************************************************/
 
@@ -64,6 +69,7 @@ void FThumbnailSection::PreDraw(FTrackEditorThumbnail& Thumbnail, FLevelEditorVi
 	}
 }
 
+
 void FThumbnailSection::PostDraw(FTrackEditorThumbnail& Thumbnail, FLevelEditorViewportClient& ViewportClient, FSceneViewport& SceneViewport)
 {
 	TSharedPtr<ISequencer> Sequencer = SequencerPtr.Pin();
@@ -83,6 +89,7 @@ bool FThumbnailSection::AreSectionsConnected() const
 	return true;
 }
 
+
 TSharedRef<SWidget> FThumbnailSection::GenerateSectionWidget()
 {
 	return SNew(SBox)
@@ -99,6 +106,7 @@ TSharedRef<SWidget> FThumbnailSection::GenerateSectionWidget()
 		];
 }
 
+
 void FThumbnailSection::BuildSectionContextMenu(FMenuBuilder& MenuBuilder, const FGuid& ObjectBinding)
 {
 	MenuBuilder.BeginSection(NAME_None, LOCTEXT("ViewMenuText", "View"));
@@ -160,6 +168,7 @@ void FThumbnailSection::BuildSectionContextMenu(FMenuBuilder& MenuBuilder, const
 	MenuBuilder.EndSection();
 }
 
+
 float FThumbnailSection::GetSectionGripSize() const
 {
 	return ThumbnailSectionConstants::SectionGripSize;
@@ -191,6 +200,7 @@ FText FThumbnailSection::GetSectionTitle() const
 	return FText::GetEmpty();
 }
 
+
 int32 FThumbnailSection::OnPaintSection( FSequencerSectionPainter& InPainter ) const
 {
 	if (!GetDefault<UMovieSceneUserThumbnailSettings>()->bDrawThumbnails)
@@ -271,4 +281,5 @@ void FThumbnailSection::Tick(const FGeometry& AllottedGeometry, const FGeometry&
 	}
 }
 
+
 #undef LOCTEXT_NAMESPACE
diff --git a/Engine/Source/Editor/MovieSceneTools/Private/Sections/ThumbnailSection.h b/Engine/Source/Editor/MovieSceneTools/Private/Sections/ThumbnailSection.h
--- a/Engine/Source/Editor/MovieSceneTools/Private/Sections/ThumbnailSection.h
+++ b/Engine/Source/Editor/MovieSceneTools/Private/Sections/ThumbnailSection.h
@@ -2,18 +2,23 @@
 
 #pragma once
 
+#include "ISequencerSection.h"
 #include "TrackEditorThumbnail.h"
 
+
 class FTrackEditorThumbnailPool;
 class IMenu;
 class ISectionLayoutBuilder;
+class SWidget;
 class UMovieSceneSection;
 
+
 struct FThumbnailCameraSettings
 {
 	float AspectRatio;
 };
 
+
 /**
  * Thumbnail section, which paints and ticks the appropriate section.
  */
@@ -52,7 +57,9 @@ public:
 	/** Set this thumbnail section to draw a single thumbnail at the specified time */
 	virtual void SetSingleTime(float GlobalTime) = 0;
 
-	// ISequencerSection interface
+public:
+
+	//~ ISequencerSection interface
 
 	virtual bool AreSectionsConnected() const override;
 	virtual void GenerateSectionLayout(ISectionLayoutBuilder& LayoutBuilder) const override { }
@@ -65,7 +72,10 @@ public:
 	virtual int32 OnPaintSection( FSequencerSectionPainter& InPainter ) const override;
 	virtual void Tick( const FGeometry& AllottedGeometry, const FGeometry& ParentGeometry, const double InCurrentTime, const float InDeltaTime ) override;
 
-	// IThumbnailClient interface
+public:
+
+	//~ IThumbnailClient interface
+
 	virtual void PreDraw(FTrackEditorThumbnail& Thumbnail, FLevelEditorViewportClient& ViewportClient, FSceneViewport& SceneViewport) override;
 	virtual void PostDraw(FTrackEditorThumbnail& Thumbnail, FLevelEditorViewportClient& ViewportClient, FSceneViewport& SceneViewport) override;
 
@@ -74,6 +84,8 @@ protected:
 	/** Called to force a redraw of this section's thumbnails */
 	void RedrawThumbnails();
 
+protected:
+
 	/** The section we are visualizing. */
 	UMovieSceneSection* Section;
 
diff --git a/Engine/Source/Editor/MovieSceneTools/Private/Sections/VectorPropertySection.cpp b/Engine/Source/Editor/MovieSceneTools/Private/Sections/VectorPropertySection.cpp
--- a/Engine/Source/Editor/MovieSceneTools/Private/Sections/VectorPropertySection.cpp
+++ b/Engine/Source/Editor/MovieSceneTools/Private/Sections/VectorPropertySection.cpp
@@ -1,35 +1,40 @@
 // Copyright 1998-2016 Epic Games, Inc. All Rights Reserved.
 
 #include "MovieSceneToolsPrivatePCH.h"
-#include "VectorPropertySection.h"
+#include "FloatCurveKeyArea.h"
 #include "MovieSceneVectorSection.h"
+#include "VectorPropertySection.h"
 
-void FVectorPropertySection::GenerateSectionLayout( class ISectionLayoutBuilder& LayoutBuilder ) const
+
+void FVectorPropertySection::GenerateSectionLayout(class ISectionLayoutBuilder& LayoutBuilder) const
 {
-	UMovieSceneVectorSection* VectorSection = Cast<UMovieSceneVectorSection>( &SectionObject );
+	UMovieSceneVectorSection* VectorSection = Cast<UMovieSceneVectorSection>(&SectionObject);
 	int32 ChannelsUsed = VectorSection->GetChannelsUsed();
-	check( ChannelsUsed >= 2 && ChannelsUsed <= 4 );
+	check(ChannelsUsed >= 2 && ChannelsUsed <= 4);
+
+	XKeyArea = MakeShareable(new FFloatCurveKeyArea(&VectorSection->GetCurve(0), VectorSection));
+	LayoutBuilder.AddKeyArea("Vector.X", NSLOCTEXT("FVectorPropertySection", "XArea", "X"), XKeyArea.ToSharedRef());
 
-	XKeyArea = MakeShareable( new FFloatCurveKeyArea( &VectorSection->GetCurve( 0 ), VectorSection ) );
-	LayoutBuilder.AddKeyArea( "Vector.X", NSLOCTEXT( "FVectorPropertySection", "XArea", "X" ), XKeyArea.ToSharedRef() );
+	YKeyArea = MakeShareable(new FFloatCurveKeyArea(&VectorSection->GetCurve(1), VectorSection));
+	LayoutBuilder.AddKeyArea("Vector.Y", NSLOCTEXT("FVectorPropertySection", "YArea", "Y"), YKeyArea.ToSharedRef());
 
-	YKeyArea = MakeShareable( new FFloatCurveKeyArea( &VectorSection->GetCurve( 1 ), VectorSection ) );
-	LayoutBuilder.AddKeyArea( "Vector.Y", NSLOCTEXT( "FVectorPropertySection", "YArea", "Y" ), YKeyArea.ToSharedRef() );
-	if ( ChannelsUsed >= 3 )
+	if (ChannelsUsed >= 3)
 	{
-		ZKeyArea = MakeShareable( new FFloatCurveKeyArea( &VectorSection->GetCurve( 2 ), VectorSection ) );
-		LayoutBuilder.AddKeyArea( "Vector.Z", NSLOCTEXT( "FVectorPropertySection", "ZArea", "Z" ), ZKeyArea.ToSharedRef() );
+		ZKeyArea = MakeShareable(new FFloatCurveKeyArea(&VectorSection->GetCurve(2), VectorSection));
+		LayoutBuilder.AddKeyArea("Vector.Z", NSLOCTEXT("FVectorPropertySection", "ZArea", "Z"), ZKeyArea.ToSharedRef());
 	}
-	if ( ChannelsUsed >= 4 )
+
+	if (ChannelsUsed >= 4)
 	{
-		WKeyArea = MakeShareable( new FFloatCurveKeyArea( &VectorSection->GetCurve( 3 ), VectorSection ) );
-		LayoutBuilder.AddKeyArea( "Vector.W", NSLOCTEXT( "FVectorPropertySection", "WArea", "W" ), WKeyArea.ToSharedRef() );
+		WKeyArea = MakeShareable(new FFloatCurveKeyArea(&VectorSection->GetCurve(3), VectorSection));
+		LayoutBuilder.AddKeyArea("Vector.W", NSLOCTEXT("FVectorPropertySection", "WArea", "W"), WKeyArea.ToSharedRef());
 	}
 }
 
-void FVectorPropertySection::SetIntermediateValue( FPropertyChangedParams PropertyChangedParams )
+
+void FVectorPropertySection::SetIntermediateValue(FPropertyChangedParams PropertyChangedParams)
 {
-	const UStructProperty* StructProp = Cast<const UStructProperty>( PropertyChangedParams.PropertyPath.Last() );
+	const UStructProperty* StructProp = Cast<const UStructProperty>(PropertyChangedParams.PropertyPath.Last());
 
 	FName StructName = StructProp->Struct->GetFName();
 
@@ -38,39 +43,42 @@ void FVectorPropertySection::SetIntermediateValue( FPropertyChangedParams Proper
 	bool bIsVector4 = StructName == NAME_Vector4;
 
 	// Get the vector value from the property
-	if ( bIsVector2D )
+	if (bIsVector2D)
 	{
 		FVector2D Vector =  PropertyChangedParams.GetPropertyValue<FVector2D>();
 		XKeyArea->SetIntermediateValue(Vector.X);
 		YKeyArea->SetIntermediateValue(Vector.Y);
 	}
-	else if ( bIsVector )
+	else if (bIsVector)
 	{
 		FVector Vector = PropertyChangedParams.GetPropertyValue<FVector>();
-		XKeyArea->SetIntermediateValue( Vector.X );
-		YKeyArea->SetIntermediateValue( Vector.Y );
-		ZKeyArea->SetIntermediateValue( Vector.Z );
+		XKeyArea->SetIntermediateValue(Vector.X);
+		YKeyArea->SetIntermediateValue(Vector.Y);
+		ZKeyArea->SetIntermediateValue(Vector.Z);
 	}
-	else // if ( bIsVector4 )
+	else // if (bIsVector4)
 	{
 		FVector4 Vector = PropertyChangedParams.GetPropertyValue<FVector4>();
-		XKeyArea->SetIntermediateValue( Vector.X );
-		YKeyArea->SetIntermediateValue( Vector.Y );
-		ZKeyArea->SetIntermediateValue( Vector.Z );
-		WKeyArea->SetIntermediateValue( Vector.W );
+		XKeyArea->SetIntermediateValue(Vector.X);
+		YKeyArea->SetIntermediateValue(Vector.Y);
+		ZKeyArea->SetIntermediateValue(Vector.Z);
+		WKeyArea->SetIntermediateValue(Vector.W);
 	}
 }
 
+
 void FVectorPropertySection::ClearIntermediateValue()
 {
 	XKeyArea->ClearIntermediateValue();
 	YKeyArea->ClearIntermediateValue();
-	if ( ZKeyArea.IsValid() )
+
+	if (ZKeyArea.IsValid())
 	{
 		ZKeyArea->ClearIntermediateValue();
 	}
-	if ( WKeyArea.IsValid() )
+
+	if (WKeyArea.IsValid())
 	{
 		WKeyArea->ClearIntermediateValue();
 	}
-}
\ No newline at end of file
+}
diff --git a/Engine/Source/Editor/MovieSceneTools/Private/Sections/VectorPropertySection.h b/Engine/Source/Editor/MovieSceneTools/Private/Sections/VectorPropertySection.h
--- a/Engine/Source/Editor/MovieSceneTools/Private/Sections/VectorPropertySection.h
+++ b/Engine/Source/Editor/MovieSceneTools/Private/Sections/VectorPropertySection.h
@@ -2,26 +2,31 @@
 
 #pragma once
 
-#include "FloatCurveKeyArea.h"
+#include "PropertySection.h"
+
+
+class FFloatCurveKeyArea;
+class UMovieSceneSection;
 
 
 /**
  * An implementation of vector property sections.
  */
-class FVectorPropertySection : public FPropertySection
+class FVectorPropertySection
+	: public FPropertySection
 {
 public:
 
 	FVectorPropertySection(UMovieSceneSection& InSectionObject, const FText& SectionName)
-		: FPropertySection( InSectionObject, SectionName )
+		: FPropertySection(InSectionObject, SectionName)
 	{ }
 
 public:
 
-	// FPropertySectionInterface
+	//~ FPropertySectionInterface
 
-	virtual void GenerateSectionLayout( class ISectionLayoutBuilder& LayoutBuilder ) const override;
-	virtual void SetIntermediateValue( FPropertyChangedParams PropertyChangedParams ) override;
+	virtual void GenerateSectionLayout(class ISectionLayoutBuilder& LayoutBuilder) const override;
+	virtual void SetIntermediateValue(FPropertyChangedParams PropertyChangedParams) override;
 	virtual void ClearIntermediateValue() override;
 
 private:
diff --git a/Engine/Source/Editor/MovieSceneTools/Private/Sections/VisibilityPropertySection.cpp b/Engine/Source/Editor/MovieSceneTools/Private/Sections/VisibilityPropertySection.cpp
--- a/Engine/Source/Editor/MovieSceneTools/Private/Sections/VisibilityPropertySection.cpp
+++ b/Engine/Source/Editor/MovieSceneTools/Private/Sections/VisibilityPropertySection.cpp
@@ -1,9 +1,11 @@
 // Copyright 1998-2016 Epic Games, Inc. All Rights Reserved.
 
 #include "MovieSceneToolsPrivatePCH.h"
+#include "BoolKeyArea.h"
 #include "VisibilityPropertySection.h"
 
-void FVisibilityPropertySection::SetIntermediateValue( FPropertyChangedParams PropertyChangedParams )
+
+void FVisibilityPropertySection::SetIntermediateValue(FPropertyChangedParams PropertyChangedParams)
 {
-	KeyArea->SetIntermediateValue( !PropertyChangedParams.GetPropertyValue<bool>() );
+	KeyArea->SetIntermediateValue(!PropertyChangedParams.GetPropertyValue<bool>());
 }
diff --git a/Engine/Source/Editor/MovieSceneTools/Private/TrackEditorThumbnail/TrackEditorThumbnail.cpp b/Engine/Source/Editor/MovieSceneTools/Private/TrackEditorThumbnail/TrackEditorThumbnail.cpp
--- a/Engine/Source/Editor/MovieSceneTools/Private/TrackEditorThumbnail/TrackEditorThumbnail.cpp
+++ b/Engine/Source/Editor/MovieSceneTools/Private/TrackEditorThumbnail/TrackEditorThumbnail.cpp
@@ -7,15 +7,21 @@
 #include "RenderCore.h"
 #include "RHIStaticStates.h"
 #include "RendererInterface.h"
+#include "TrackEditorThumbnail.h"
+#include "TrackEditorThumbnailPool.h"
+
 
 namespace TrackEditorThumbnailConstants
 {
 	const double ThumbnailFadeInDuration = 0.25f;
 }
 
-class FThumbnailViewportClient : public FLevelEditorViewportClient
+
+class FThumbnailViewportClient
+	: public FLevelEditorViewportClient
 {
 public:
+
 	FThumbnailViewportClient() : FLevelEditorViewportClient(nullptr) {}
 
 	float CurrentWorldTime, DeltaWorldTime;
@@ -60,6 +66,7 @@ FTrackEditorThumbnail::~FTrackEditorThumbnail()
 	DestroyTexture();
 }
 
+
 void FTrackEditorThumbnail::DestroyTexture()
 {
 	if (Texture)
@@ -75,6 +82,7 @@ void FTrackEditorThumbnail::DestroyTexture()
 	}
 }
 
+
 /* FTrackEditorThumbnail interface
  *****************************************************************************/
 
@@ -190,17 +198,20 @@ void FTrackEditorThumbnail::DrawThumbnail()
 	OnDraw.ExecuteIfBound(*this);
 }
 
+
 void FTrackEditorThumbnail::SetupFade(const TSharedRef<SWidget>& InWidget)
 {
 	FadeInCurve.PlayReverse(InWidget);
 	FadeInCurve.Pause();
 }
 
+
 void FTrackEditorThumbnail::PlayFade()
 {
 	FadeInCurve.Resume();
 }
 
+
 float FTrackEditorThumbnail::GetFadeInCurve() const 
 {
 	return FadeInCurve.GetLerp();
@@ -253,6 +264,7 @@ FTrackEditorThumbnailCache::~FTrackEditorThumbnailCache()
 	}
 }
 
+
 void FTrackEditorThumbnailCache::SetSingleReferenceFrame(TOptional<float> InReferenceFrame)
 {
 	CurrentCache.SingleReferenceFrame = InReferenceFrame;
@@ -278,6 +290,7 @@ void FTrackEditorThumbnailCache::Update(const TRange<float>& NewRange, const TRa
 	PreviousCache.SingleReferenceFrame = CurrentCache.SingleReferenceFrame;
 }
 
+
 FIntPoint FTrackEditorThumbnailCache::CalculateTextureSize() const
 {
 	UCameraComponent* CameraComponent = InternalViewportClient->GetCameraComponentForView();
@@ -310,7 +323,7 @@ FIntPoint FTrackEditorThumbnailCache::CalculateTextureSize() const
 		case EThumbnailQuality::Draft: 	Scale = 0.5f; 	break;
 		case EThumbnailQuality::Best: 	Scale = 2.f; 	break;
 		default: 						Scale = 1.f; 	break;
-	}
+}
 
 	return FIntPoint(
 		FMath::RoundToInt(X * Scale),
@@ -332,6 +345,7 @@ bool FTrackEditorThumbnailCache::ShouldRegenerateEverything() const
 	return PreviousCache.DesiredSize != CurrentCache.DesiredSize || !FMath::IsNearlyEqual(PreviousScale, CurrentScale, Threshold);
 }
 
+
 void FTrackEditorThumbnailCache::DrawViewportThumbnail(FTrackEditorThumbnail& TrackEditorThumbnail)
 {
 	if (InternalViewportScene.IsValid())
@@ -375,6 +389,7 @@ void FTrackEditorThumbnailCache::DrawViewportThumbnail(FTrackEditorThumbnail& Tr
 	}
 }
 
+
 void FTrackEditorThumbnailCache::Revalidate(double InCurrentTime)
 {
 	if (CurrentCache == PreviousCache && !bForceRedraw && !bNeedsNewThumbnails)
@@ -404,7 +419,7 @@ void FTrackEditorThumbnailCache::Revalidate(double InCurrentTime)
 		ThumbnailPool.Pin()->RemoveThumbnailsNeedingRedraw(Thumbnails);
 		Thumbnails.Reset();
 	}
-
+	
 	if (InCurrentTime - LastComputationTime > 0.25f)
 	{
 		ComputeNewThumbnails();
@@ -441,6 +456,7 @@ void FTrackEditorThumbnailCache::ComputeNewThumbnails()
 	bNeedsNewThumbnails = false;
 }
 
+
 void FTrackEditorThumbnailCache::UpdateSingleThumbnail()
 {
 	Thumbnails.Reset();
@@ -460,6 +476,7 @@ void FTrackEditorThumbnailCache::UpdateSingleThumbnail()
 	ThumbnailsNeedingRedraw.Add(NewThumbnail);
 }
 
+
 void FTrackEditorThumbnailCache::UpdateFilledThumbnails()
 {
 	// Remove any thumbnails from the front of the array that aren't in the actual time range of this section (we keep stuff around outside of the visible range)
@@ -525,6 +542,7 @@ void FTrackEditorThumbnailCache::UpdateFilledThumbnails()
 	}
 }
 
+
 void FTrackEditorThumbnailCache::GenerateFront(const TRange<float>& Boundary)
 {
 	if (!Thumbnails.Num())
@@ -560,6 +578,7 @@ void FTrackEditorThumbnailCache::GenerateFront(const TRange<float>& Boundary)
 	}
 }
 
+
 void FTrackEditorThumbnailCache::GenerateBack(const TRange<float>& Boundary)
 {
 	const float TimePerPx = CurrentCache.TimeRange.Size<float>() / CurrentCache.AllottedSize.X;
@@ -600,6 +619,7 @@ void FTrackEditorThumbnailCache::GenerateBack(const TRange<float>& Boundary)
 	}
 }
 
+
 void FTrackEditorThumbnailCache::Setup()
 {
 	if (!InternalViewportClient.IsValid())
diff --git a/Engine/Source/Editor/MovieSceneTools/Private/TrackEditorThumbnail/TrackEditorThumbnailPool.cpp b/Engine/Source/Editor/MovieSceneTools/Private/TrackEditorThumbnail/TrackEditorThumbnailPool.cpp
--- a/Engine/Source/Editor/MovieSceneTools/Private/TrackEditorThumbnail/TrackEditorThumbnailPool.cpp
+++ b/Engine/Source/Editor/MovieSceneTools/Private/TrackEditorThumbnail/TrackEditorThumbnailPool.cpp
@@ -1,6 +1,8 @@
 // Copyright 1998-2016 Epic Games, Inc. All Rights Reserved.
 
 #include "MovieSceneToolsPrivatePCH.h"
+#include "TrackEditorThumbnail.h"
+#include "TrackEditorThumbnailPool.h"
 
 
 /* FShotSequencerSection structors
@@ -24,11 +26,12 @@ void FTrackEditorThumbnailPool::AddThumbnailsNeedingRedraw(const TArray<TSharedP
 	bNeedsSort = true;
 }
 
+
 bool FTrackEditorThumbnailPool::DrawThumbnails()
 {
 	int32 ThumbnailsDrawn = 0;
 
-	// Apply sorting if necessary
+	// apply sorting if necessary
 	if (bNeedsSort)
 	{
 		auto SortFunc = [](const TSharedPtr<FTrackEditorThumbnail>& A, const TSharedPtr<FTrackEditorThumbnail>& B){
@@ -38,14 +41,16 @@ bool FTrackEditorThumbnailPool::DrawThumbnails()
 			}
 			return A->SortOrder < B->SortOrder;
 		};
+
 		ThumbnailsNeedingDraw.Sort(SortFunc);
 		bNeedsSort = false;
 	}
 
 	const double CurrentTime = FSlateApplication::Get().GetCurrentTime();
 
-	// Only allow drawing if we're not waiting on other thumbnails
+	// only allow drawing if we're not waiting on other thumbnails
 	bool bAllowDraw = !ThumbnailsBeingDrawn.Num();
+
 	for (int32 ThumbnailIndex = ThumbnailsBeingDrawn.Num() - 1; ThumbnailIndex >= 0; --ThumbnailIndex)
 	{
 		if (ThumbnailsBeingDrawn[ThumbnailIndex]->bHasFinishedDrawing)
@@ -75,17 +80,15 @@ bool FTrackEditorThumbnailPool::DrawThumbnails()
 			const float MaxFramerate = 90.f;
 
 			const float AverageDeltaTime = FSlateApplication::Get().GetAverageDeltaTime();
-
-			const float Lerp = FMath::Max(((1.f/AverageDeltaTime) - MinFramerate), 0.f) / (MaxFramerate - MinFramerate);
+			const float Lerp = FMath::Max(((1.f / AverageDeltaTime) - MinFramerate), 0.f) / (MaxFramerate - MinFramerate);
 			const float NumToDrawPerS = FMath::Lerp(MinThumbnailsPerS, MaxThumbnailsPerS, Lerp);
 
 			const float TimeThreshold = 1.f / NumToDrawPerS;
-
 			const float TimeElapsed = CurrentTime - TimeOfLastDraw;
 
 			if (!FMath::IsNearlyEqual(float(CurrentTime - TimeOfLastUpdate), AverageDeltaTime, AverageDeltaTime*2.f))
 			{
-				// Don't generate thumbnails if we haven't had an update within a reasonable time - assume some blocking task
+				// don't generate thumbnails if we haven't had an update within a reasonable time - assume some blocking task
 				TimeOfLastDraw = CurrentTime;
 			}
 			else
diff --git a/Engine/Source/Editor/MovieSceneTools/Private/TrackEditors/CameraCutTrackEditor.cpp b/Engine/Source/Editor/MovieSceneTools/Private/TrackEditors/CameraCutTrackEditor.cpp
--- a/Engine/Source/Editor/MovieSceneTools/Private/TrackEditors/CameraCutTrackEditor.cpp
+++ b/Engine/Source/Editor/MovieSceneTools/Private/TrackEditors/CameraCutTrackEditor.cpp
@@ -21,6 +21,7 @@
 #include "Editor/SceneOutliner/Public/SceneOutliner.h"
 #include "Editor/LevelEditor/Public/LevelEditor.h"
 #include "SceneOutlinerPublicTypes.h"
+#include "TrackEditorThumbnailPool.h"
 
 #define LOCTEXT_NAMESPACE "FCameraCutTrackEditor"
 
