Remove Simplygon warning when generating engine DDC. Now that we distribute the engine integration for Simplygon (but not the proprietary DLL itself), it's normal for the DLL to not be present.

diff --git a/Engine/Source/Developer/SimplygonMeshReduction/Private/SimplygonMeshReduction.cpp b/Engine/Source/Developer/SimplygonMeshReduction/Private/SimplygonMeshReduction.cpp
--- a/Engine/Source/Developer/SimplygonMeshReduction/Private/SimplygonMeshReduction.cpp
+++ b/Engine/Source/Developer/SimplygonMeshReduction/Private/SimplygonMeshReduction.cpp
@@ -517,7 +517,7 @@ public:
 		// If the DLL just doesn't exist, fail gracefully. Licensees and Subscribers will not necessarily have Simplygon.
 		if( !FPaths::FileExists(DllFilename) )
 		{
-			UE_LOG(LogSimplygon,Warning,TEXT("Simplygon DLL not present - disabling."));
+			UE_LOG(LogSimplygon,Log,TEXT("Simplygon DLL not present - disabling."));
 			return NULL;
 		}
 
