fix for the simplygon initialization failure message to properly print the error code

diff --git a/Engine/Source/Developer/SimplygonMeshReduction/Private/SimplygonMeshReduction.cpp b/Engine/Source/Developer/SimplygonMeshReduction/Private/SimplygonMeshReduction.cpp
--- a/Engine/Source/Developer/SimplygonMeshReduction/Private/SimplygonMeshReduction.cpp
+++ b/Engine/Source/Developer/SimplygonMeshReduction/Private/SimplygonMeshReduction.cpp
@@ -705,7 +705,7 @@ public:
 
 			if (InitResult != SimplygonSDK::SG_ERROR_NOERROR && InitResult != SimplygonSDK::SG_ERROR_ALREADYINITIALIZED)
 			{
-				UE_LOG(LogSimplygon, Warning, TEXT("Failed to initialize Simplygon. Error: %d."), Result);
+				UE_LOG(LogSimplygon, Warning, TEXT("Failed to initialize Simplygon. Error: %d."), InitResult);
 				FPlatformProcess::FreeDllHandle(GSimplygonSDKDLLHandle);
 				GSimplygonSDKDLLHandle = nullptr;
 			}
