Simplygon changes

diff --git a/Engine/Source/Developer/SimplygonMeshReduction/Private/SimplygonMeshReduction.cpp b/Engine/Source/Developer/SimplygonMeshReduction/Private/SimplygonMeshReduction.cpp
--- a/Engine/Source/Developer/SimplygonMeshReduction/Private/SimplygonMeshReduction.cpp
+++ b/Engine/Source/Developer/SimplygonMeshReduction/Private/SimplygonMeshReduction.cpp
@@ -627,7 +627,7 @@ public:
 		if (FParse::Param(FCommandLine::Get(), TEXT("NoSimplygon")))
 		{
 			//The user specified that simplygon should not be used
-			UE_LOG(LogSimplygon, Warning, TEXT("Simplygon is disabled with -NoSimplygon flag"));
+			UE_LOG(LogSimplygon, Log, TEXT("Simplygon is disabled with -NoSimplygon flag"));
 			return  NULL;
 		}
 		//@third party BEGIN SIMPLYGON
@@ -643,7 +643,7 @@ public:
 #endif
 		if (!FPaths::FileExists(DllPath))
 		{
-			UE_LOG(LogSimplygon, Warning, TEXT("Could not find d3dcompiler_47 DLL, which is required for loading Simplygon."));
+			UE_LOG(LogSimplygon, Log, TEXT("Could not find d3dcompiler_47 DLL, which is required for loading Simplygon."));
 			return NULL;
 		}
 		else
@@ -683,7 +683,7 @@ public:
 		if ((GetInterfaceVersionSimplygonSDK == NULL) || (InitializeSimplygonSDK == NULL))
 		{
 			// Couldn't find the functions we need.  
-			UE_LOG(LogSimplygon,Warning,TEXT("Failed to acquire Simplygon DLL exports."));
+			UE_LOG(LogSimplygon, Log,TEXT("Failed to acquire Simplygon DLL exports."));
 			FPlatformProcess::FreeDllHandle( GSimplygonSDKDLLHandle );
 			GSimplygonSDKDLLHandle = NULL;
 			return NULL;
@@ -698,8 +698,8 @@ public:
 
 		if (SUPPORTED_SIMPLYGON_VERSION.Compare(MajorVersion) != 0)
 		{
-			UE_LOG(LogSimplygon, Warning, TEXT("Simplygon version doesn't match the version expected by the Simplygon UE4 integration"));
-			UE_LOG(LogSimplygon, Warning, TEXT("Min version %s, found version %s"), *SUPPORTED_SIMPLYGON_VERSION, ANSI_TO_TCHAR(SimplygonSDK::GetHeaderVersion()));
+			UE_LOG(LogSimplygon, Log, TEXT("Simplygon version doesn't match the version expected by the Simplygon UE4 integration"));
+			UE_LOG(LogSimplygon, Log, TEXT("Min version %s, found version %s"), *SUPPORTED_SIMPLYGON_VERSION, ANSI_TO_TCHAR(SimplygonSDK::GetHeaderVersion()));
 			FPlatformProcess::FreeDllHandle(GSimplygonSDKDLLHandle);
 			GSimplygonSDKDLLHandle = NULL;
 			return NULL;
@@ -711,7 +711,7 @@ public:
 		GetInterfaceVersionSimplygonSDK(VersionHash);
 		if (FCStringAnsi::Strcmp(VersionHash, SimplygonSDK::GetInterfaceVersionHash()) != 0)
 		{
-			UE_LOG(LogSimplygon,Warning,TEXT("Library version mismatch. Header=%s Lib=%s"),ANSI_TO_TCHAR(SimplygonSDK::GetInterfaceVersionHash()),ANSI_TO_TCHAR(VersionHash));
+			UE_LOG(LogSimplygon, Log,TEXT("Library version mismatch. Header=%s Lib=%s"),ANSI_TO_TCHAR(SimplygonSDK::GetInterfaceVersionHash()),ANSI_TO_TCHAR(VersionHash));
 			FPlatformProcess::FreeDllHandle(GSimplygonSDKDLLHandle);
 			GSimplygonSDKDLLHandle = NULL;
 			return NULL;
@@ -764,7 +764,7 @@ public:
 
 			if (InitResult != SimplygonSDK::SG_ERROR_NOERROR && InitResult != SimplygonSDK::SG_ERROR_ALREADYINITIALIZED)
 			{
-				UE_LOG(LogSimplygon, Warning, TEXT("Failed to initialize Simplygon. Return code: %d."), InitResult);
+				UE_LOG(LogSimplygon, Log, TEXT("Failed to initialize Simplygon. Return code: %d."), InitResult);
 				FPlatformProcess::FreeDllHandle(GSimplygonSDKDLLHandle);
 				GSimplygonSDKDLLHandle = nullptr;
 			}
@@ -775,7 +775,7 @@ public:
 		}
 		else
 		{
-			UE_LOG(LogSimplygon, Warning, TEXT("Failed to load Simplygon license file."));
+			UE_LOG(LogSimplygon, Log, TEXT("Failed to load Simplygon license file."));
 		}
 
 		return Result;
