Support simplygon reduction API for MeshDescription

diff --git a/Engine/Source/Developer/MeshDescriptionOperations/Private/MeshDescriptionOperations.cpp b/Engine/Source/Developer/MeshDescriptionOperations/Private/MeshDescriptionOperations.cpp
--- a/Engine/Source/Developer/MeshDescriptionOperations/Private/MeshDescriptionOperations.cpp
+++ b/Engine/Source/Developer/MeshDescriptionOperations/Private/MeshDescriptionOperations.cpp
@@ -63,7 +63,7 @@ namespace MeshDescriptionOperationNamespace
 //////////////////////////////////////////////////////////////////////////
 // Converters
 
-void FMeshDescriptionOperations::ConvertHardEdgesToSmoothGroup(const FMeshDescription& SourceMeshDescription, FRawMesh& DestinationRawMesh)
+void FMeshDescriptionOperations::ConvertHardEdgesToSmoothGroup(const FMeshDescription& SourceMeshDescription, TArray<uint32>& FaceSmoothingMasks)
 {
 	TMap<FPolygonID, uint32> PolygonSmoothGroup;
 	PolygonSmoothGroup.Reserve(SourceMeshDescription.Polygons().GetArraySize());
@@ -73,9 +73,10 @@ void FMeshDescriptionOperations::ConvertHardEdgesToSmoothGroup(const FMeshDescri
 	TMap < FPolygonID, uint32> PolygonAvoidances;
 
 	TEdgeAttributesConstRef<bool> EdgeHardnesses = SourceMeshDescription.EdgeAttributes().GetAttributesRef<bool>(MeshAttribute::Edge::IsHard);
-
+	int32 TriangleCount = 0;
 	for (const FPolygonID PolygonID : SourceMeshDescription.Polygons().GetElementIDs())
 	{
+		TriangleCount += SourceMeshDescription.GetPolygonTriangles(PolygonID).Num();
 		if (ConsumedPolygons[PolygonID.GetValue()])
 		{
 			continue;
@@ -173,7 +174,8 @@ void FMeshDescriptionOperations::ConvertHardEdgesToSmoothGroup(const FMeshDescri
 			ConsumedPolygons[CurrentPolygonID.GetValue()] = true;
 		}
 	}
-	//Now we have to put the data into the RawMesh
+	//Set the smooth group in the FaceSmoothingMasks parameter
+	check(FaceSmoothingMasks.Num() == TriangleCount);
 	int32 TriangleIndex = 0;
 	for (const FPolygonID PolygonID : SourceMeshDescription.Polygons().GetElementIDs())
 	{
@@ -181,7 +183,7 @@ void FMeshDescriptionOperations::ConvertHardEdgesToSmoothGroup(const FMeshDescri
 		const TArray<FMeshTriangle>& Triangles = SourceMeshDescription.GetPolygonTriangles(PolygonID);
 		for (const FMeshTriangle& MeshTriangle : Triangles)
 		{
-			DestinationRawMesh.FaceSmoothingMasks[TriangleIndex++] = PolygonSmoothValue;
+			FaceSmoothingMasks[TriangleIndex++] = PolygonSmoothValue;
 		}
 	}
 }
@@ -324,7 +326,7 @@ void FMeshDescriptionOperations::ConvertToRawMesh(const FMeshDescription& Source
 		}
 	}
 	//Convert the smoothgroup
-	ConvertHardEdgesToSmoothGroup(SourceMeshDescription, DestinationRawMesh);
+	ConvertHardEdgesToSmoothGroup(SourceMeshDescription, DestinationRawMesh.FaceSmoothingMasks);
 }
 
 //We want to fill the FMeshDescription vertex position mesh attribute with the FRawMesh vertex position
diff --git a/Engine/Source/Developer/MeshDescriptionOperations/Public/MeshDescriptionOperations.h b/Engine/Source/Developer/MeshDescriptionOperations/Public/MeshDescriptionOperations.h
--- a/Engine/Source/Developer/MeshDescriptionOperations/Public/MeshDescriptionOperations.h
+++ b/Engine/Source/Developer/MeshDescriptionOperations/Public/MeshDescriptionOperations.h
@@ -66,8 +66,7 @@ public:
 	/** Remove the UV channel at the given index from the MeshDescription. */
 	static bool RemoveUVChannel(FMeshDescription& MeshDescription, int32 UVChannelIndex);
 
-
-	static void ConvertHardEdgesToSmoothGroup(const FMeshDescription& SourceMeshDescription, FRawMesh& DestinationRawMesh);
+	static void ConvertHardEdgesToSmoothGroup(const FMeshDescription& SourceMeshDescription, TArray<uint32>& FaceSmoothingMasks);
 
 	static void ConvertSmoothGroupToHardEdges(const TArray<uint32>& FaceSmoothingMasks, FMeshDescription& DestinationMeshDescription);
 };
\ No newline at end of file
diff --git a/Engine/Source/Developer/SimplygonMeshReduction/Private/SimplygonMeshReduction.cpp b/Engine/Source/Developer/SimplygonMeshReduction/Private/SimplygonMeshReduction.cpp
--- a/Engine/Source/Developer/SimplygonMeshReduction/Private/SimplygonMeshReduction.cpp
+++ b/Engine/Source/Developer/SimplygonMeshReduction/Private/SimplygonMeshReduction.cpp
@@ -28,6 +28,11 @@
 #include "MeshMergeData.h"
 #include "Rendering/MultiSizeIndexContainer.h"
 
+#include "MeshDescription.h"
+#include "MeshAttributes.h"
+#include "MeshAttributeArray.h"
+#include "MeshDescriptionOperations.h"
+
 // Standard Simplygon channels have some issues with extracting color data back from simplification, 
 // so we use this workaround with user channels
 static const char* USER_MATERIAL_CHANNEL_METALLIC = "UserMetallic";
@@ -567,9 +572,45 @@ public:
 		float& OutMaxDeviation,
 		const FMeshDescription& InMesh,
 		const FOverlappingCorners& InOverlappingCorners,
-		const struct FMeshReductionSettings& ReductionSettings
+		const struct FMeshReductionSettings& InSettings
 	) override
 	{
+		SimplygonSDK::spGeometryData GeometryData = CreateGeometryFromRawMesh(InMesh);
+		check(GeometryData);
+
+		SimplygonSDK::spScene Scene = SDK->CreateScene();
+
+		SimplygonSDK::spSceneMesh Mesh = SDK->CreateSceneMesh();
+		Mesh->SetGeometry(GeometryData);
+		Mesh->SetName(TCHAR_TO_ANSI(*FString::Printf(TEXT("UnrealMesh"))));
+		Scene->GetRootNode()->AddChild(Mesh);
+
+		SimplygonSDK::spReductionProcessor ReductionProcessor = SDK->CreateReductionProcessor();
+		ReductionProcessor->AddObserver(&EventHandler, SimplygonSDK::SG_EVENT_PROGRESS);
+		ReductionProcessor->SetScene(Scene);
+
+		SimplygonSDK::spRepairSettings RepairSettings = ReductionProcessor->GetRepairSettings();
+		RepairSettings->SetWeldDist(InSettings.WeldingThreshold);
+		RepairSettings->SetTjuncDist(InSettings.WeldingThreshold);
+
+		SimplygonSDK::spReductionSettings ReductionSettings = ReductionProcessor->GetReductionSettings();
+		SetReductionSettings(ReductionSettings, InSettings, GeometryData->GetTriangleCount());
+
+		//Set visibility settings
+		SimplygonSDK::spVisibilitySettings VisibilitySettings = ReductionProcessor->GetVisibilitySettings();
+		VisibilitySettings->SetCullOccludedGeometry(InSettings.bCullOccluded);
+		int32 TempAggressiveness = InSettings.VisibilityAggressiveness + 1; //+1 because there is an offset in aggressiveness options
+		VisibilitySettings->SetVisibilityWeightsPower(TempAggressiveness);
+		VisibilitySettings->SetUseVisibilityWeightsInReducer(InSettings.bVisibilityAided);
+
+		SimplygonSDK::spNormalCalculationSettings NormalSettings = ReductionProcessor->GetNormalCalculationSettings();
+		SetNormalSettings(NormalSettings, InSettings);
+
+		ReductionProcessor->RunProcessing();
+
+		SimplygonSDK::spSceneMesh ReducedMesh = SimplygonSDK::Cast<SimplygonSDK::ISceneMesh>(Scene->GetRootNode()->GetChild(0));
+		CreateRawMeshFromGeometry(OutReducedMesh, ReducedMesh->GetGeometry(), WINDING_Keep);
+		OutMaxDeviation = ReductionProcessor->GetResultDeviation();
 	}
 
 	virtual bool ReduceSkeletalMesh(
@@ -1570,6 +1611,330 @@ private:
 		}
 	}
 
+	SimplygonSDK::spGeometryData CreateGeometryFromRawMesh(const FMeshDescription& RawMesh)
+	{
+		TVertexAttributesConstRef<FVector> VertexPositions = RawMesh.VertexAttributes().GetAttributesRef<FVector>(MeshAttribute::Vertex::Position);
+		TEdgeAttributesConstRef<bool> EdgeHardnesses = RawMesh.EdgeAttributes().GetAttributesRef<bool>(MeshAttribute::Edge::IsHard);
+		TEdgeAttributesConstRef<float> EdgeCreaseSharpnesses = RawMesh.EdgeAttributes().GetAttributesRef<float>(MeshAttribute::Edge::CreaseSharpness);
+		TPolygonGroupAttributesConstRef<FName> PolygonGroupImportedMaterialSlotNames = RawMesh.PolygonGroupAttributes().GetAttributesRef<FName>(MeshAttribute::PolygonGroup::ImportedMaterialSlotName);
+		TVertexInstanceAttributesConstRef<FVector> VertexInstanceNormals = RawMesh.VertexInstanceAttributes().GetAttributesRef<FVector>(MeshAttribute::VertexInstance::Normal);
+		TVertexInstanceAttributesConstRef<FVector> VertexInstanceTangents = RawMesh.VertexInstanceAttributes().GetAttributesRef<FVector>(MeshAttribute::VertexInstance::Tangent);
+		TVertexInstanceAttributesConstRef<float> VertexInstanceBinormalSigns = RawMesh.VertexInstanceAttributes().GetAttributesRef<float>(MeshAttribute::VertexInstance::BinormalSign);
+		TVertexInstanceAttributesConstRef<FVector4> VertexInstanceColors = RawMesh.VertexInstanceAttributes().GetAttributesRef<FVector4>(MeshAttribute::VertexInstance::Color);
+		TVertexInstanceAttributesConstRef<FVector2D> VertexInstanceUVs = RawMesh.VertexInstanceAttributes().GetAttributesRef<FVector2D>(MeshAttribute::VertexInstance::TextureCoordinate);
+
+		int32 NumVertices = RawMesh.Vertices().Num();
+		int32 NumWedges = RawMesh.VertexInstances().Num();
+
+		if (NumWedges == 0)
+		{
+			return NULL;
+		}
+
+		int32 NumTris = 0;
+		for (const FPolygonID& PolygonID : RawMesh.Polygons().GetElementIDs())
+		{
+			NumTris += RawMesh.GetPolygon(PolygonID).Triangles.Num();
+		}
+
+		TArray<uint32> FaceSmoothingMasks;
+		FaceSmoothingMasks.AddZeroed(NumTris);
+		FMeshDescriptionOperations::ConvertHardEdgesToSmoothGroup(RawMesh, FaceSmoothingMasks);
+
+		SimplygonSDK::spGeometryData GeometryData = SDK->CreateGeometryData();
+		GeometryData->SetVertexCount(NumVertices);
+		GeometryData->SetTriangleCount(NumTris);
+
+		TMap<FVertexID, int32> VertexIDToDstVertexIndex;
+		VertexIDToDstVertexIndex.Reserve(NumVertices);
+		int32 VertexCount = 0;
+		SimplygonSDK::spRealArray Positions = GeometryData->GetCoords();
+		for (const FVertexID& VertexID : RawMesh.Vertices().GetElementIDs())
+		{
+			FVector TempPos = VertexPositions[VertexID];
+			TempPos = GetConversionMatrix().TransformPosition(TempPos);
+			Positions->SetTuple(VertexID.GetValue(), (float*)&TempPos);
+			VertexIDToDstVertexIndex.Add(VertexID, VertexCount);
+			VertexCount++;
+		}
+
+		//Prepare the tex coord
+		TArray<SimplygonSDK::spRealArray> TexCoordsArray;
+		for (int32 TexCoordIndex = 0; TexCoordIndex < VertexInstanceUVs.GetNumIndices(); ++TexCoordIndex)
+		{
+			GeometryData->AddTexCoords(TexCoordIndex);
+			SimplygonSDK::spRealArray TexCoords = GeometryData->GetTexCoords(TexCoordIndex);
+			check(TexCoords->GetTupleSize() == 2);
+			TexCoordsArray.Add(TexCoords);
+		}
+
+		//Prepare the vertex color
+		GeometryData->AddColors(0);
+		SimplygonSDK::spRealArray LinearColors = GeometryData->GetColors(0);
+		check(LinearColors);
+		check(LinearColors->GetTupleSize() == 4);
+
+		//Prepare the tangent space
+		GeometryData->AddTangents(0);
+		GeometryData->AddNormals();
+		SimplygonSDK::spRealArray Tangents = GeometryData->GetTangents(0);
+		SimplygonSDK::spRealArray Bitangents = GeometryData->GetBitangents(0);
+		SimplygonSDK::spRealArray Normals = GeometryData->GetNormals();
+
+		SimplygonSDK::spRidArray Indices = GeometryData->GetVertexIds();
+
+		GeometryData->AddMaterialIds();
+		SimplygonSDK::spRidArray MaterialIndices = GeometryData->GetMaterialIds();
+
+		uint32 DstVertexIndex = 0;
+		uint32 DstTriangleIndex = 0;
+		for (const FPolygonID& PolygonID : RawMesh.Polygons().GetElementIDs())
+		{
+			for (const FMeshTriangle& Triangle : RawMesh.GetPolygon(PolygonID).Triangles)
+			{
+				for (int32 Corner = 0; Corner < 3; ++Corner)
+				{
+					const FVertexInstanceID VertexInstanceID = Triangle.GetVertexInstanceID(Corner);
+					//Add all the per indice data
+					Indices->SetItem(DstVertexIndex, VertexIDToDstVertexIndex[RawMesh.GetVertexInstanceVertex(VertexInstanceID)]);
+
+					//UVs
+					for (int32 TexCoordIndex = 0; TexCoordIndex < VertexInstanceUVs.GetNumIndices(); ++TexCoordIndex)
+					{
+						FVector2D Vector2D = VertexInstanceUVs.Get(VertexInstanceID, TexCoordIndex);
+						TexCoordsArray[TexCoordIndex]->SetTuple(DstVertexIndex, (float*)&Vector2D);
+					}
+
+					//Colors
+					LinearColors->SetTuple(DstVertexIndex, (float*)&VertexInstanceColors[VertexInstanceID]);
+
+					//Tangents
+					FVector TempTangent = VertexInstanceTangents[VertexInstanceID];
+					TempTangent = GetConversionMatrix().TransformVector(TempTangent);
+					Tangents->SetTuple(DstVertexIndex, (float*)&TempTangent);
+
+					FVector TempBitangent = FVector::CrossProduct(VertexInstanceNormals[VertexInstanceID], VertexInstanceTangents[VertexInstanceID]).GetSafeNormal() * VertexInstanceBinormalSigns[VertexInstanceID];
+					TempBitangent = GetConversionMatrix().TransformVector(TempBitangent);
+					Bitangents->SetTuple(DstVertexIndex, (float*)&TempBitangent);
+
+					FVector TempNormal = VertexInstanceNormals[VertexInstanceID];
+					TempNormal = GetConversionMatrix().TransformVector(TempNormal);
+					Normals->SetTuple(DstVertexIndex, (float*)&TempNormal);
+
+					//Increment the indice index
+					DstVertexIndex++;
+				}
+
+				// Per-triangle data.
+
+				//Materials
+				const FPolygonGroupID& PolygonGroupID = RawMesh.GetPolygonPolygonGroup(PolygonID);
+				MaterialIndices->SetItem(DstTriangleIndex, PolygonGroupID.GetValue());
+
+				//Smooth group
+				GeometryData->AddGroupIds();
+				SimplygonSDK::spRidArray GroupIds = GeometryData->GetGroupIds();
+				GroupIds->SetItem(DstTriangleIndex, FaceSmoothingMasks[DstTriangleIndex]);
+
+				DstTriangleIndex++;
+			}
+		}
+
+		return GeometryData;
+	}
+
+	void CreateRawMeshFromGeometry(FMeshDescription& OutRawMesh, const SimplygonSDK::spGeometryData& GeometryData, EWindingMode WindingMode)
+	{
+		check(GeometryData);
+
+		SimplygonSDK::spRealArray Positions = GeometryData->GetCoords();
+		SimplygonSDK::spRidArray Indices = GeometryData->GetVertexIds();
+		SimplygonSDK::spRidArray MaterialIndices = GeometryData->GetMaterialIds();
+		SimplygonSDK::spRidArray GroupIds = GeometryData->GetGroupIds();
+		SimplygonSDK::spRealArray LinearColors = GeometryData->GetColors(0);
+		SimplygonSDK::spRealArray Normals = GeometryData->GetNormals();
+		SimplygonSDK::spRealArray Tangents = GeometryData->GetTangents(0);
+		SimplygonSDK::spRealArray Bitangents = GeometryData->GetBitangents(0);
+
+		check(Positions);
+		check(Indices);
+
+		FMeshDescription& RawMesh = OutRawMesh;
+
+		TVertexAttributesRef<FVector> VertexPositions = RawMesh.VertexAttributes().GetAttributesRef<FVector>(MeshAttribute::Vertex::Position);
+		TEdgeAttributesRef<bool> EdgeHardnesses = RawMesh.EdgeAttributes().GetAttributesRef<bool>(MeshAttribute::Edge::IsHard);
+		TEdgeAttributesRef<float> EdgeCreaseSharpnesses = RawMesh.EdgeAttributes().GetAttributesRef<float>(MeshAttribute::Edge::CreaseSharpness);
+		TPolygonGroupAttributesRef<FName> PolygonGroupImportedMaterialSlotNames = RawMesh.PolygonGroupAttributes().GetAttributesRef<FName>(MeshAttribute::PolygonGroup::ImportedMaterialSlotName);
+		TVertexInstanceAttributesRef<FVector> VertexInstanceNormals = RawMesh.VertexInstanceAttributes().GetAttributesRef<FVector>(MeshAttribute::VertexInstance::Normal);
+		TVertexInstanceAttributesRef<FVector> VertexInstanceTangents = RawMesh.VertexInstanceAttributes().GetAttributesRef<FVector>(MeshAttribute::VertexInstance::Tangent);
+		TVertexInstanceAttributesRef<float> VertexInstanceBinormalSigns = RawMesh.VertexInstanceAttributes().GetAttributesRef<float>(MeshAttribute::VertexInstance::BinormalSign);
+		TVertexInstanceAttributesRef<FVector4> VertexInstanceColors = RawMesh.VertexInstanceAttributes().GetAttributesRef<FVector4>(MeshAttribute::VertexInstance::Color);
+		TVertexInstanceAttributesRef<FVector2D> VertexInstanceUVs = RawMesh.VertexInstanceAttributes().GetAttributesRef<FVector2D>(MeshAttribute::VertexInstance::TextureCoordinate);
+
+		const bool bReverseWinding = (WindingMode == WINDING_Reverse);
+		int32 NumTris = GeometryData->GetTriangleCount();
+		int32 NumWedges = NumTris * 3;
+		int32 NumVertices = GeometryData->GetVertexCount();
+
+		RawMesh.Empty();
+		RawMesh.ReserveNewVertices(NumVertices);
+		RawMesh.ReserveNewPolygons(NumTris);
+		RawMesh.ReserveNewVertexInstances(NumWedges);
+		RawMesh.ReserveNewEdges(NumWedges);
+
+		TMap<int32, FVertexID> GeoToRawVertexID;
+		SimplygonSDK::spRealData sgTuple = SDK->CreateRealData();
+		for (int32 VertexIndex = 0; VertexIndex < NumVertices; ++VertexIndex)
+		{
+			Positions->GetTuple(VertexIndex, sgTuple);
+			SimplygonSDK::real* vertexPos = sgTuple->GetData();
+			const FVertexID& VertexID = RawMesh.CreateVertex();
+			VertexPositions[VertexID] = GetConversionMatrix().TransformPosition(FVector(vertexPos[0], vertexPos[1], vertexPos[2]));
+			GeoToRawVertexID.Add(VertexIndex, VertexID);
+		}
+
+		//Prepare the tex coord
+		int32 TexCoordNum = 0;
+		TArray<SimplygonSDK::spRealArray> TexCoordsArray;
+		for (int32 TexCoordIndex = 0; TexCoordIndex < MAX_MESH_TEXTURE_COORDS; ++TexCoordIndex, ++TexCoordNum)
+		{
+			SimplygonSDK::spRealArray TexCoords = GeometryData->GetTexCoords(TexCoordIndex);
+			if (TexCoords == nullptr)
+			{
+				break;
+			}
+			check(TexCoords->GetTupleSize() == 2);
+			TexCoordsArray.Add(TexCoords);
+		}
+		VertexInstanceUVs.SetNumIndices(TexCoordNum);
+
+		//Prepare the polygongroup
+		TMap<int32, FPolygonGroupID> GeoToRawMaterial;
+		for (int32 TriIndex = 0; TriIndex < NumTris; ++TriIndex)
+		{
+			int32 MaterialIndex = MaterialIndices->GetItem(TriIndex);
+			if (!GeoToRawMaterial.Contains(MaterialIndex))
+			{
+				const FPolygonGroupID PolygonGroupID(MaterialIndex);
+				RawMesh.CreatePolygonGroupWithID(PolygonGroupID);
+				PolygonGroupImportedMaterialSlotNames[PolygonGroupID] = FName(*FString(TEXT("SimplygonMeshReduction_") + FString::FromInt(PolygonGroupID.GetValue())));
+				GeoToRawMaterial.Add(MaterialIndex, PolygonGroupID);
+			}
+		}
+
+		bool bHasZeroTangent = false;
+		bool bHasZeroNormal = false;
+		for (int32 TriIndex = 0; TriIndex < NumTris; ++TriIndex)
+		{
+			FVertexID VertexIndexes[3];
+			FVertexInstanceID VertexInstanceIDs[3];
+			for (int32 CornerIndex = 0; CornerIndex < 3; ++CornerIndex)
+			{
+				const uint32 SrcIndex = (TriIndex * 3) + (bReverseWinding ? (2 - CornerIndex) : CornerIndex);
+				const FVertexID VertexID = GeoToRawVertexID[Indices->GetItem(SrcIndex)];
+				const FVertexInstanceID VertexInstanceID = RawMesh.CreateVertexInstance(VertexID);
+				VertexIndexes[CornerIndex] = VertexID;
+				VertexInstanceIDs[CornerIndex] = VertexInstanceID;
+				//Texture coordinnates
+				for (int32 TexCoordIndex = 0; TexCoordIndex < TexCoordNum; ++TexCoordIndex)
+				{
+					TexCoordsArray[TexCoordIndex]->GetTuple(SrcIndex, sgTuple);
+					SimplygonSDK::real* sgTextCoors = sgTuple->GetData();
+					VertexInstanceUVs.Set(VertexInstanceID, TexCoordIndex, FVector2D(sgTextCoors[0], sgTextCoors[1]));
+				}
+
+				//Vertex Color
+				LinearColors->GetTuple(SrcIndex, sgTuple);
+				SimplygonSDK::real* sgVertexColor = sgTuple->GetData();
+				VertexInstanceColors[VertexInstanceID] = FVector4(sgVertexColor[0], sgVertexColor[1], sgVertexColor[2], sgVertexColor[3]);
+
+				//Vertex Tangents
+				FVector Tangent(0.0f);
+				FVector BiTangent(0.0f);
+				FVector Normal(0.0f);
+
+				Normals->GetTuple(SrcIndex, sgTuple);
+				SimplygonSDK::real* sgNormal = sgTuple->GetData();
+				Normal = GetConversionMatrix().TransformVector(FVector(sgNormal[0], sgNormal[1], sgNormal[2]));
+				if (Normal.IsNearlyZero())
+				{
+					bHasZeroNormal = true;
+				}
+
+				if (Tangents && Bitangents)
+				{
+					Tangents->GetTuple(SrcIndex, sgTuple);
+					SimplygonSDK::real* sgTangents = sgTuple->GetData();
+					Tangent = GetConversionMatrix().TransformVector(FVector(sgTangents[0], sgTangents[1], sgTangents[2]));
+
+					Bitangents->GetTuple(SrcIndex, sgTuple);
+					SimplygonSDK::real* sgBiTangents = sgTuple->GetData();
+					BiTangent = GetConversionMatrix().TransformVector(FVector(sgBiTangents[0], sgBiTangents[1], sgBiTangents[2]));
+				}
+				else
+				{
+					bHasZeroTangent = true;
+					Tangent = FVector(0.0f);
+					BiTangent = FVector(0.0f);
+				}
+				VertexInstanceTangents[VertexInstanceID] = Tangent;
+				VertexInstanceBinormalSigns[VertexInstanceID] = GetBasisDeterminantSign(Tangent.GetSafeNormal(), BiTangent.GetSafeNormal(), Normal.GetSafeNormal());
+				VertexInstanceNormals[VertexInstanceID] = Normal;
+			}
+			//Create a polygon from this triangle
+			TArray<FMeshDescription::FContourPoint> Contours;
+			for (int32 Corner = 0; Corner < 3; ++Corner)
+			{
+				int32 ContourPointIndex = Contours.AddDefaulted();
+				FMeshDescription::FContourPoint& ContourPoint = Contours[ContourPointIndex];
+				//Find the matching edge ID
+				uint32 CornerIndices[2];
+				CornerIndices[0] = (Corner + 0) % 3;
+				CornerIndices[1] = (Corner + 1) % 3;
+
+				FVertexID EdgeVertexIDs[2];
+				EdgeVertexIDs[0] = VertexIndexes[CornerIndices[0]];
+				EdgeVertexIDs[1] = VertexIndexes[CornerIndices[1]];
+
+				FEdgeID MatchEdgeId = RawMesh.GetVertexPairEdge(EdgeVertexIDs[0], EdgeVertexIDs[1]);
+				if (MatchEdgeId == FEdgeID::Invalid)
+				{
+					MatchEdgeId = RawMesh.CreateEdge(EdgeVertexIDs[0], EdgeVertexIDs[1]);
+					EdgeHardnesses[MatchEdgeId] = false;
+					EdgeCreaseSharpnesses[MatchEdgeId] = 0.0f;
+				}
+				ContourPoint.EdgeID = MatchEdgeId;
+				ContourPoint.VertexInstanceID = VertexInstanceIDs[CornerIndices[0]];
+			}
+			// Insert a polygon into the mesh
+			const FPolygonID NewPolygonID = RawMesh.CreatePolygon(GeoToRawMaterial[MaterialIndices->GetItem(TriIndex)], Contours);
+			//Triangulate the polygon
+			FMeshPolygon& Polygon = RawMesh.GetPolygon(NewPolygonID);
+			RawMesh.ComputePolygonTriangulation(NewPolygonID, Polygon.Triangles);
+		}
+
+		TArray<uint32> FaceSmoothingMasks;
+		FaceSmoothingMasks.AddZeroed(NumTris);
+		if (GroupIds)
+		{
+			for (int32 TriIndex = 0; TriIndex < NumTris; ++TriIndex)
+			{
+				FaceSmoothingMasks[TriIndex] = GroupIds->GetItem(TriIndex);
+			}
+		}
+		FMeshDescriptionOperations::ConvertSmoothGroupToHardEdges(FaceSmoothingMasks, RawMesh);
+
+		if (bHasZeroTangent | bHasZeroNormal)
+		{
+			FMeshDescriptionOperations::ETangentOptions TangentOption = (FMeshDescriptionOperations::ETangentOptions)(FMeshDescriptionOperations::ETangentOptions::BlendOverlappingNormals | FMeshDescriptionOperations::ETangentOptions::IgnoreDegenerateTriangles);
+			if (bHasZeroNormal)
+			{
+				FMeshDescriptionOperations::CreateNormals(RawMesh, TangentOption, false);
+			}
+			FMeshDescriptionOperations::CreateMikktTangents(RawMesh, TangentOption);
+		}
+	}
+
 	void SetReductionSettings(SimplygonSDK::spReductionSettings ReductionSettings, const FMeshReductionSettings& Settings, int32 NumTris)
 	{
 		float MaxDeviation = Settings.MaxDeviation > 0.0f ? Settings.MaxDeviation : SimplygonSDK::REAL_MAX;
diff --git a/Engine/Source/Developer/SimplygonMeshReduction/SimplygonMeshReduction.Build.cs b/Engine/Source/Developer/SimplygonMeshReduction/SimplygonMeshReduction.Build.cs
--- a/Engine/Source/Developer/SimplygonMeshReduction/SimplygonMeshReduction.Build.cs
+++ b/Engine/Source/Developer/SimplygonMeshReduction/SimplygonMeshReduction.Build.cs
@@ -16,6 +16,8 @@ public class SimplygonMeshReduction : ModuleRules
 				"Engine",
 				"RenderCore",
                 "RawMesh",
+                "MeshDescription",
+                "MeshDescriptionOperations",
                 "MaterialUtilities",
                 "MeshBoneReduction",
                 "RHI",
