Deprecate FRawMesh - Remove all FRawMesh dependencies from simplygon reduction plugin. This code was unaccessible.

diff --git a/Engine/Source/Developer/SimplygonMeshReduction/Private/SimplygonMeshReduction.cpp b/Engine/Source/Developer/SimplygonMeshReduction/Private/SimplygonMeshReduction.cpp
--- a/Engine/Source/Developer/SimplygonMeshReduction/Private/SimplygonMeshReduction.cpp
+++ b/Engine/Source/Developer/SimplygonMeshReduction/Private/SimplygonMeshReduction.cpp
@@ -1,7 +1,6 @@
 // Copyright 1998-2018 Epic Games, Inc. All Rights Reserved.
 
 #include "CoreMinimal.h"
-#include "RawMesh.h"
 #include "IMeshReductionInterfaces.h"
 #include "MeshUtilities.h"
 #include "MaterialUtilities.h"
@@ -1188,385 +1187,6 @@ private:
 		UE_LOG(LogSimplygon, Display, TEXT("Initialized with Simplygon %s"), *VersionString);
 	}
 
-	SimplygonSDK::spGeometryData CreateGeometryFromRawMesh(const FRawMesh& RawMesh)
-	{
-		int32 NumVertices = RawMesh.VertexPositions.Num();
-		int32 NumWedges = RawMesh.WedgeIndices.Num();
-		int32 NumTris = NumWedges / 3;
-
-		if (NumWedges == 0)
-		{
-			return NULL;
-		}
-
-		SimplygonSDK::spGeometryData GeometryData = SDK->CreateGeometryData();
-		GeometryData->SetVertexCount(NumVertices);
-		GeometryData->SetTriangleCount(NumTris);
-
-		SimplygonSDK::spRealArray Positions = GeometryData->GetCoords();
-		for (int32 VertexIndex = 0; VertexIndex < NumVertices; ++VertexIndex)
-		{
-			FVector TempPos = RawMesh.VertexPositions[VertexIndex];
-			TempPos = GetConversionMatrix().TransformPosition(TempPos);
-			Positions->SetTuple(VertexIndex, (float*)&TempPos);
-		}
-
-		SimplygonSDK::spRidArray Indices = GeometryData->GetVertexIds();
-		for (int32 WedgeIndex = 0; WedgeIndex < NumWedges; ++WedgeIndex)
-		{
-			Indices->SetItem(WedgeIndex, RawMesh.WedgeIndices[WedgeIndex]);
-		}
-
-		for(int32 TexCoordIndex = 0; TexCoordIndex < MAX_MESH_TEXTURE_COORDS; ++TexCoordIndex)
-		{
-			if (RawMesh.WedgeTexCoords[TexCoordIndex].Num() == NumWedges)
-			{
-				GeometryData->AddTexCoords(TexCoordIndex);
-				SimplygonSDK::spRealArray TexCoords = GeometryData->GetTexCoords(TexCoordIndex);
-				check(TexCoords->GetTupleSize() == 2);
-				for (int32 WedgeIndex = 0; WedgeIndex < NumWedges; ++WedgeIndex)
-				{
-					TexCoords->SetTuple(WedgeIndex, (float*)&RawMesh.WedgeTexCoords[TexCoordIndex][WedgeIndex]);
-				}
-			}
-		}
-
-		if (RawMesh.WedgeColors.Num() == NumWedges)
-		{
-			GeometryData->AddColors(0); 
-			SimplygonSDK::spRealArray LinearColors = GeometryData->GetColors(0);
-			check(LinearColors);
-			check(LinearColors->GetTupleSize() == 4);
-			for (int32 WedgeIndex = 0; WedgeIndex < NumWedges; ++WedgeIndex)
-			{
-				FLinearColor LinearColor(RawMesh.WedgeColors[WedgeIndex]);
-				LinearColors->SetTuple(WedgeIndex, (float*)&LinearColor);
-			}
-		}
-
-		if (RawMesh.WedgeTangentZ.Num() == NumWedges)
-		{
-			if (RawMesh.WedgeTangentX.Num() == NumWedges && RawMesh.WedgeTangentY.Num() == NumWedges)
-			{
-				GeometryData->AddTangents(0);
-				SimplygonSDK::spRealArray Tangents = GeometryData->GetTangents(0);
-				for (int32 WedgeIndex = 0; WedgeIndex < NumWedges; ++WedgeIndex)
-				{
-					FVector TempTangent = RawMesh.WedgeTangentX[WedgeIndex];
-					TempTangent = GetConversionMatrix().TransformVector(TempTangent);
-					Tangents->SetTuple(WedgeIndex, (float*)&TempTangent);
-				}
-
-				GeometryData->AddBitangents(0);
-				SimplygonSDK::spRealArray Bitangents = GeometryData->GetBitangents(0);
-				for (int32 WedgeIndex = 0; WedgeIndex < NumWedges; ++WedgeIndex)
-				{
-					FVector TempBitangent = RawMesh.WedgeTangentY[WedgeIndex];
-					TempBitangent = GetConversionMatrix().TransformVector(TempBitangent);
-					Bitangents->SetTuple(WedgeIndex, (float*)&TempBitangent);
-				}
-			}
-			GeometryData->AddNormals();
-			SimplygonSDK::spRealArray Normals = GeometryData->GetNormals();
-			for (int32 WedgeIndex = 0; WedgeIndex < NumWedges; ++WedgeIndex)
-			{
-				FVector TempNormal = RawMesh.WedgeTangentZ[WedgeIndex];
-				TempNormal = GetConversionMatrix().TransformVector(TempNormal);
-				Normals->SetTuple(WedgeIndex, (float*)&TempNormal);
-			}
-		}
-
-		// Per-triangle data.
-		GeometryData->AddMaterialIds();
-		SimplygonSDK::spRidArray MaterialIndices = GeometryData->GetMaterialIds();
-		for (int32 TriIndex = 0; TriIndex < NumTris; ++TriIndex)
-		{
-			MaterialIndices->SetItem(TriIndex, RawMesh.FaceMaterialIndices[TriIndex]);
-		}
-
-		GeometryData->AddGroupIds();
-		SimplygonSDK::spRidArray GroupIds = GeometryData->GetGroupIds();
-		for (int32 TriIndex = 0; TriIndex < NumTris; ++TriIndex)
-		{
-			GroupIds->SetItem(TriIndex, RawMesh.FaceSmoothingMasks[TriIndex]);
-		}
-		
-		return GeometryData;
-
-	}
-
-	// This is a copy of CreateGeometryFromRawMesh with additional features for material LOD.
-	SimplygonSDK::spGeometryData CreateGeometryFromRawMesh(const FRawMesh& RawMesh, const TArray<FBox2D> & TextureBounds, const TArray<FVector2D>& InTexCoords)
-	{
-		int32 NumVertices = RawMesh.VertexPositions.Num();
-		int32 NumWedges = RawMesh.WedgeIndices.Num();
-		int32 NumTris = NumWedges / 3;
-
-		if (NumWedges == 0)
-		{
-			return NULL;
-		}
-
-		SimplygonSDK::spGeometryData GeometryData = SDK->CreateGeometryData();
-		GeometryData->SetVertexCount(NumVertices);
-		GeometryData->SetTriangleCount(NumTris);
-
-		SimplygonSDK::spRealArray Positions = GeometryData->GetCoords();
-		for (int32 VertexIndex = 0; VertexIndex < NumVertices; ++VertexIndex)
-		{
-			FVector TempPos = RawMesh.VertexPositions[VertexIndex];
-			TempPos = GetConversionMatrix().TransformPosition(TempPos);
-			Positions->SetTuple(VertexIndex, (float*)&TempPos);
-		}
-
-		SimplygonSDK::spRidArray Indices = GeometryData->GetVertexIds();
-		for (int32 WedgeIndex = 0; WedgeIndex < NumWedges; ++WedgeIndex)
-		{
-			Indices->SetItem(WedgeIndex, RawMesh.WedgeIndices[WedgeIndex]);
-		}
-
-		for(int32 TexCoordIndex = 0; TexCoordIndex < MAX_MESH_TEXTURE_COORDS; ++TexCoordIndex)
-		{
-			const TArray<FVector2D>& SrcTexCoords = (TexCoordIndex == 0 && InTexCoords.Num() == NumWedges) ? InTexCoords : RawMesh.WedgeTexCoords[TexCoordIndex];
-			if (SrcTexCoords.Num() == NumWedges)
-			{
-				GeometryData->AddTexCoords(TexCoordIndex);
-				SimplygonSDK::spRealArray TexCoords = GeometryData->GetTexCoords(TexCoordIndex);
-				check(TexCoords->GetTupleSize() == 2);
-				int32 WedgeIndex = 0;
-				for (int32 TriIndex = 0; TriIndex < NumTris; ++TriIndex)
-				{
-					int32 MaterialIndex = RawMesh.FaceMaterialIndices[TriIndex];
-					// Compute texture bounds for current material.
-					float MinU = 0, ScaleU = 1;
-					float MinV = 0, ScaleV = 1;
-					
-					if (TextureBounds.IsValidIndex(MaterialIndex) && TexCoordIndex == 0 && InTexCoords.Num() == 0)
-					{
-						const FBox2D& Bounds = TextureBounds[MaterialIndex];
-						if (Bounds.GetArea() > 0)
-						{
-							MinU = Bounds.Min.X;
-							MinV = Bounds.Min.Y;
-							ScaleU = 1.0f / (Bounds.Max.X - Bounds.Min.X);
-							ScaleV = 1.0f / (Bounds.Max.Y - Bounds.Min.Y);
-						}
-					}
-					for (int32 CornerIndex = 0; CornerIndex < 3; ++CornerIndex, ++WedgeIndex)
-					{
-						const FVector2D& TexCoord = SrcTexCoords[WedgeIndex];
-						float UV[2];
-						UV[0] = (TexCoord.X - MinU) * ScaleU;
-						UV[1] = (TexCoord.Y - MinV) * ScaleV;
-						TexCoords->SetTuple(WedgeIndex, UV);
-					}
-				}
-			}
-		}
-
-		if (RawMesh.WedgeColors.Num() == NumWedges)
-		{
-			GeometryData->AddColors(0); 
-			SimplygonSDK::spRealArray LinearColors = GeometryData->GetColors(0);
-			check(LinearColors);
-			check(LinearColors->GetTupleSize() == 4);
-			for (int32 WedgeIndex = 0; WedgeIndex < NumWedges; ++WedgeIndex)
-			{
-				FLinearColor LinearColor(RawMesh.WedgeColors[WedgeIndex]);
-				LinearColors->SetTuple(WedgeIndex, (float*)&LinearColor);
-			}
-		}
-
-		if (RawMesh.WedgeTangentZ.Num() == NumWedges)
-		{
-			if (RawMesh.WedgeTangentX.Num() == NumWedges && RawMesh.WedgeTangentY.Num() == NumWedges)
-			{
-				GeometryData->AddTangents(0);
-				SimplygonSDK::spRealArray Tangents = GeometryData->GetTangents(0);
-				for (int32 WedgeIndex = 0; WedgeIndex < NumWedges; ++WedgeIndex)
-				{
-					FVector TempTangent = RawMesh.WedgeTangentX[WedgeIndex];
-					TempTangent = GetConversionMatrix().TransformVector(TempTangent);
-					Tangents->SetTuple(WedgeIndex, (float*)&TempTangent);
-				}
-
-				GeometryData->AddBitangents(0);
-				SimplygonSDK::spRealArray Bitangents = GeometryData->GetBitangents(0);
-				for (int32 WedgeIndex = 0; WedgeIndex < NumWedges; ++WedgeIndex)
-				{
-					FVector TempBitangent = RawMesh.WedgeTangentY[WedgeIndex];
-					TempBitangent = GetConversionMatrix().TransformVector(TempBitangent);
-					Bitangents->SetTuple(WedgeIndex, (float*)&TempBitangent);
-				}
-			}
-			GeometryData->AddNormals();
-			SimplygonSDK::spRealArray Normals = GeometryData->GetNormals();
-			for (int32 WedgeIndex = 0; WedgeIndex < NumWedges; ++WedgeIndex)
-			{
-				FVector TempNormal = RawMesh.WedgeTangentZ[WedgeIndex];
-				TempNormal = GetConversionMatrix().TransformVector(TempNormal);
-				Normals->SetTuple(WedgeIndex, (float*)&TempNormal);
-			}
-		}
-
-		// Per-triangle data.
-		GeometryData->AddMaterialIds();
-		SimplygonSDK::spRidArray MaterialIndices = GeometryData->GetMaterialIds();
-		for (int32 TriIndex = 0; TriIndex < NumTris; ++TriIndex)
-		{
-			MaterialIndices->SetItem(TriIndex, RawMesh.FaceMaterialIndices[TriIndex]);
-		}
-
-		GeometryData->AddGroupIds();
-		SimplygonSDK::spRidArray GroupIds = GeometryData->GetGroupIds();
-		for (int32 TriIndex = 0; TriIndex < NumTris; ++TriIndex)
-		{
-			GroupIds->SetItem(TriIndex, RawMesh.FaceSmoothingMasks[TriIndex]);
-		}
-		
-		return GeometryData;
-
-	}
-	 
-
-	void CreateRawMeshFromGeometry(FRawMesh& OutRawMesh, const SimplygonSDK::spGeometryData& GeometryData, EWindingMode WindingMode)
-	{
-		check(GeometryData);
-
-		SimplygonSDK::spRealArray Positions = GeometryData->GetCoords();
-		SimplygonSDK::spRidArray Indices = GeometryData->GetVertexIds();
-		SimplygonSDK::spRidArray MaterialIndices = GeometryData->GetMaterialIds();
-		SimplygonSDK::spRidArray GroupIds = GeometryData->GetGroupIds();
-		SimplygonSDK::spRealArray LinearColors = GeometryData->GetColors(0);
-		SimplygonSDK::spRealArray Normals = GeometryData->GetNormals();
-		SimplygonSDK::spRealArray Tangents = GeometryData->GetTangents(0);
-		SimplygonSDK::spRealArray Bitangents = GeometryData->GetBitangents(0);
-
-		check(Positions);
-		check(Indices);
-
-		FRawMesh& RawMesh = OutRawMesh;
-		const bool bReverseWinding = (WindingMode == WINDING_Reverse);
-		int32 NumTris = GeometryData->GetTriangleCount();
-		int32 NumWedges = NumTris * 3;
-		int32 NumVertices = GeometryData->GetVertexCount();
-
-		RawMesh.VertexPositions.Empty(NumVertices);
-		RawMesh.VertexPositions.AddUninitialized(NumVertices);
-		SimplygonSDK::spRealData sgTuple = SDK->CreateRealData();
-		for (int32 VertexIndex = 0; VertexIndex < NumVertices; ++VertexIndex)
-		{
-			//Positions->GetTuple(VertexIndex, (float*)&RawMesh.VertexPositions[VertexIndex]);
-			Positions->GetTuple(VertexIndex, sgTuple);
-			SimplygonSDK::real* vertexPos = sgTuple->GetData();
-			RawMesh.VertexPositions[VertexIndex] = GetConversionMatrix().TransformPosition(FVector(vertexPos[0], vertexPos[1], vertexPos[2]));
-		}
-
-		RawMesh.WedgeIndices.Empty(NumWedges);
-		RawMesh.WedgeIndices.AddUninitialized(NumWedges);
-		for (int32 TriIndex = 0; TriIndex < NumTris; ++TriIndex)
-		{
-			for (int32 CornerIndex = 0; CornerIndex < 3; ++CornerIndex)
-			{
-				const uint32 DestIndex = bReverseWinding ? (2 - CornerIndex) : CornerIndex;
-				RawMesh.WedgeIndices[TriIndex * 3 + DestIndex] = Indices->GetItem(TriIndex * 3 + CornerIndex);
-			}
-		}
-
-		for(int32 TexCoordIndex = 0; TexCoordIndex < MAX_MESH_TEXTURE_COORDS; ++TexCoordIndex)
-		{
-			SimplygonSDK::spRealArray TexCoords = GeometryData->GetTexCoords(TexCoordIndex);
-			if (TexCoords)
-			{
-				RawMesh.WedgeTexCoords[TexCoordIndex].Empty(NumWedges);
-				RawMesh.WedgeTexCoords[TexCoordIndex].AddUninitialized(NumWedges);
-				for (int32 TriIndex = 0; TriIndex < NumTris; ++TriIndex)
-				{
-					for (int32 CornerIndex = 0; CornerIndex < 3; ++CornerIndex)
-					{
-						const uint32 DestIndex = bReverseWinding ? (2 - CornerIndex) : CornerIndex;
-						//TexCoords->GetTuple(TriIndex * 3 + CornerIndex, (float*)&RawMesh.WedgeTexCoords[TexCoordIndex][TriIndex * 3 + DestIndex]);
-						TexCoords->GetTuple(TriIndex * 3 + CornerIndex, sgTuple);
-
-						SimplygonSDK::real* sgTextCoors = sgTuple->GetData();
-						RawMesh.WedgeTexCoords[TexCoordIndex][TriIndex * 3 + DestIndex] = FVector2D(sgTextCoors[0], sgTextCoors[1]);
-
-					}
-				}
-			}
-		}
-
-		RawMesh.WedgeColors.Empty(NumWedges);
-		RawMesh.WedgeColors.AddUninitialized(NumWedges);
-		if (LinearColors)
-		{
-			for (int32 WedgeIndex = 0; WedgeIndex < NumWedges; ++WedgeIndex)
-			{
-				//LinearColors->GetTuple(WedgeIndex, (float*)&LinearColor);
-				LinearColors->GetTuple(WedgeIndex, sgTuple);
-				SimplygonSDK::real* sgVertexColor = sgTuple->GetData();
-				FLinearColor LinearColor(sgVertexColor[0], sgVertexColor[1], sgVertexColor[2], sgVertexColor[3]);
-				RawMesh.WedgeColors[WedgeIndex] = LinearColor.ToFColor(true);
-			}
-		}
-
-		if (Normals)
-		{
-			if (Tangents && Bitangents)
-			{
-				RawMesh.WedgeTangentX.Empty(NumWedges);
-				RawMesh.WedgeTangentX.AddUninitialized(NumWedges);
-				for (int32 WedgeIndex = 0; WedgeIndex < NumWedges; ++WedgeIndex)
-				{
-					//Tangents->GetTuple(WedgeIndex, (float*)&RawMesh.WedgeTangentX[WedgeIndex]);
-					Tangents->GetTuple(WedgeIndex, sgTuple);
-					SimplygonSDK::real* sgTangents = sgTuple->GetData();
-					RawMesh.WedgeTangentX[WedgeIndex] = GetConversionMatrix().TransformVector( FVector(sgTangents[0], sgTangents[1], sgTangents[2]) );
-				}
-
-				RawMesh.WedgeTangentY.Empty(NumWedges);
-				RawMesh.WedgeTangentY.AddUninitialized(NumWedges);
-				for (int32 WedgeIndex = 0; WedgeIndex < NumWedges; ++WedgeIndex)
-				{
-					//Bitangents->GetTuple(WedgeIndex, (float*)&RawMesh.WedgeTangentY[WedgeIndex]);
-					Bitangents->GetTuple(WedgeIndex, sgTuple);
-					SimplygonSDK::real* sgBitangents = sgTuple->GetData();
-					RawMesh.WedgeTangentY[WedgeIndex] = GetConversionMatrix().TransformVector(FVector(sgBitangents[0], sgBitangents[1], sgBitangents[2]));
-				}
-			}
-
-			RawMesh.WedgeTangentZ.Empty(NumWedges);
-			RawMesh.WedgeTangentZ.AddUninitialized(NumWedges);
-			for (int32 WedgeIndex = 0; WedgeIndex < NumWedges; ++WedgeIndex)
-			{
-				//Normals->GetTuple(WedgeIndex, (float*)&RawMesh.WedgeTangentZ[WedgeIndex]);
-				Normals->GetTuple(WedgeIndex, sgTuple);
-				SimplygonSDK::real* sgNormal = sgTuple->GetData();
-				RawMesh.WedgeTangentZ[WedgeIndex] = GetConversionMatrix().TransformVector(FVector(sgNormal[0], sgNormal[1], sgNormal[2]));
-			}
-		}
-
-		RawMesh.FaceMaterialIndices.Empty(NumTris);
-		RawMesh.FaceMaterialIndices.AddZeroed(NumTris);
-		if (MaterialIndices)
-		{
-			for (int32 TriIndex = 0; TriIndex < NumTris; ++TriIndex)
-			{
-				RawMesh.FaceMaterialIndices[TriIndex] = MaterialIndices->GetItem(TriIndex);
-			}
-		}
-
-		RawMesh.FaceSmoothingMasks.Empty(NumTris);
-		RawMesh.FaceSmoothingMasks.AddZeroed(NumTris);
-		if (GroupIds)
-		{
-			for (int32 TriIndex = 0; TriIndex < NumTris; ++TriIndex)
-			{
-				RawMesh.FaceSmoothingMasks[TriIndex] = GroupIds->GetItem(TriIndex);
-			}
-		}
-	}
-
 	SimplygonSDK::spGeometryData CreateGeometryFromRawMesh(const FMeshDescription& RawMesh)
 	{
 		TVertexAttributesConstRef<FVector> VertexPositions = RawMesh.VertexAttributes().GetAttributesRef<FVector>(MeshAttribute::Vertex::Position);
@@ -1903,7 +1523,7 @@ private:
 		//Prepare the tex coord
 		int32 TexCoordNum = 0;
 		TArray<SimplygonSDK::spRealArray> TexCoordsArray;
-		for (int32 TexCoordIndex = 0; TexCoordIndex < MAX_MESH_TEXTURE_COORDS; ++TexCoordIndex, ++TexCoordNum)
+		for (int32 TexCoordIndex = 0; TexCoordIndex < MAX_MESH_TEXTURE_COORDS_MD; ++TexCoordIndex, ++TexCoordNum)
 		{
 			SimplygonSDK::spRealArray TexCoords = GeometryData->GetTexCoords(TexCoordIndex);
 			if (TexCoords == nullptr)
diff --git a/Engine/Source/Developer/SimplygonMeshReduction/SimplygonMeshReduction.Build.cs b/Engine/Source/Developer/SimplygonMeshReduction/SimplygonMeshReduction.Build.cs
--- a/Engine/Source/Developer/SimplygonMeshReduction/SimplygonMeshReduction.Build.cs
+++ b/Engine/Source/Developer/SimplygonMeshReduction/SimplygonMeshReduction.Build.cs
@@ -15,7 +15,6 @@ public class SimplygonMeshReduction : ModuleRules
 				"CoreUObject",
 				"Engine",
 				"RenderCore",
-				"RawMesh",
                 "MeshDescription",
                 "MeshDescriptionOperations",
                 "MaterialUtilities",
